Copying //UE4/Dev-Framework to //UE4/Dev-Main (Source: //UE4/Dev-Framework @ 3136612)

diff --git a/Engine/Source/Editor/SkeletonEditor/Private/SSkeletonTree.cpp b/Engine/Source/Editor/SkeletonEditor/Private/SSkeletonTree.cpp
new file mode 100644
--- /dev/null
+++ b/Engine/Source/Editor/SkeletonEditor/Private/SSkeletonTree.cpp
@@ -0,0 +1,2744 @@
+// Copyright 1998-2016 Epic Games, Inc. All Rights Reserved.
+
+
+#include "SkeletonEditorPrivatePCH.h"
+#include "SSkeletonTree.h"
+
+#include "ScopedTransaction.h"
+#include "BoneDragDropOp.h"
+#include "SocketDragDropOp.h"
+#include "SkeletonTreeCommands.h"
+#include "DragAndDrop/AssetDragDropOp.h"
+
+#include "AssetSelection.h"
+
+#include "Editor/ContentBrowser/Public/ContentBrowserModule.h"
+#include "ComponentAssetBroker.h"
+
+#include "SlateIconFinder.h"
+
+#include "Editor/UnrealEd/Public/AssetNotifications.h"
+
+#include "Animation/PreviewAssetAttachComponent.h"
+#include "AnimPreviewInstance.h"
+
+#include "Factories.h"
+#include "Developer/MeshUtilities/Public/MeshUtilities.h"
+#include "UnrealExporter.h"
+#include "SSearchBox.h"
+#include "SInlineEditableTextBlock.h"
+#include "SNotificationList.h"
+#include "NotificationManager.h"
+#include "GenericCommands.h"
+#include "Animation/BlendProfile.h"
+#include "SBlendProfilePicker.h"
+#include "IPersonaPreviewScene.h"
+#include "IDocumentation.h"
+#include "PersonaUtils.h"
+
+#define LOCTEXT_NAMESPACE "SSkeletonTree"
+
+static const FName	ColumnID_BoneLabel( "BoneName" );
+static const FName	ColumnID_RetargetingLabel( "TranslationRetargeting" );
+static const FName  ColumnID_BlendProfileLabel( "BlendProfile" );
+
+DECLARE_DELEGATE_RetVal_TwoParams(FReply, FOnDraggingBoneItem, const FGeometry&, const FPointerEvent&);
+
+//////////////////////////////////////////////////////////////////////////
+// SSkeletonTreeRow
+
+DECLARE_DELEGATE_RetVal_TwoParams(FReply, FOnDraggingTreeItem, const FGeometry&, const FPointerEvent&);
+
+typedef TSharedPtr< FDisplayedTreeRowInfo > FDisplayedTreeRowInfoPtr;
+
+class SSkeletonTreeRow
+	: public SMultiColumnTableRow< FDisplayedTreeRowInfoPtr >
+{
+public:
+
+	SLATE_BEGIN_ARGS(SSkeletonTreeRow) {}
+
+	/** The item for this row **/
+	SLATE_ARGUMENT(FDisplayedTreeRowInfoPtr, Item)
+
+	/** Pointer to the owning skeleton tree so it can be used to copy sockets, tree refresh when needed etc */
+	SLATE_ARGUMENT(TWeakPtr< SSkeletonTree >, SkeletonTree);
+
+	/** Pointer to the owning skeleton tree so it can be used to copy sockets, tree refresh when needed etc */
+	SLATE_ARGUMENT(TWeakPtr< FEditableSkeleton >, EditableSkeleton);
+
+	/** Filter text typed by the user into the parent tree's search widget */
+	SLATE_ARGUMENT(FText, FilterText);
+
+	/** Delegate for dragging items **/
+	SLATE_EVENT(FOnDraggingTreeItem, OnDraggingItem);
+
+	SLATE_END_ARGS()
+
+	void Construct(const FArguments& InArgs, const TSharedRef<STableViewBase>& InOwnerTableView);
+
+	/** Overridden from SMultiColumnTableRow.  Generates a widget for this column of the tree row. */
+	virtual TSharedRef<SWidget> GenerateWidgetForColumn(const FName& ColumnName) override;
+
+	/** Overridden from STableRow. Allows us to generate inline edit widgets. */
+	virtual void ConstructChildren(ETableViewMode::Type InOwnerTableMode, const TAttribute<FMargin>& InPadding, const TSharedRef<SWidget>& InContent) override;
+
+	/** Override OnDragEnter for drag and drop of sockets onto bones */
+	virtual void OnDragEnter(const FGeometry& MyGeometry, const FDragDropEvent& DragDropEvent) override;
+
+	/** Override OnDragLeave for drag and drop of sockets onto bones */
+	virtual void OnDragLeave(const FDragDropEvent& DragDropEvent) override;
+
+	/** Override OnDrop for drag and drop of sockets and meshes onto bones */
+	virtual FReply OnDrop(const FGeometry& MyGeometry, const FDragDropEvent& DragDropEvent) override;
+
+	/** Override DoesItemHaveChildren to add expander to in-line editors */
+	virtual int32 DoesItemHaveChildren() const override;
+
+	/** @return true if the corresponding item is expanded; false otherwise */
+	virtual bool IsItemExpanded() const override;
+
+	/** Toggle the expansion of the item associated with this row */
+	virtual void ToggleExpansion() override;
+
+protected:
+
+	/** Handler for starting a drag/drop action */
+	virtual FReply OnDragDetected(const FGeometry& MyGeometry, const FPointerEvent& MouseEvent) override;
+
+	/** Get the editable skeleton we are editing */
+	TSharedRef<FEditableSkeleton> GetEditableSkeleton() const { return EditableSkeleton.Pin().ToSharedRef();  }
+
+	/** Get the skeleton tree we are embedded in */
+	TSharedRef<SSkeletonTree>	GetSkeletonTree() const { return SkeletonTree.Pin().ToSharedRef(); }
+
+private:
+
+	/** The item this row is holding */
+	FDisplayedTreeRowInfoPtr	Item;
+
+	/** Text the user typed into the search box - used for text highlighting */
+	FText						FilterText;
+
+	/** Weak pointer to the owning skeleton tree */
+	TWeakPtr<SSkeletonTree>		SkeletonTree;
+
+	/** Weak pointer to the owning editable skeleton */
+	TWeakPtr<FEditableSkeleton>	EditableSkeleton;
+
+	/** Item that we're dragging */
+	FOnDraggingTreeItem			OnDraggingItem;
+
+	/** Was the user pressing "Alt" when the drag was started? */
+	bool						bIsAltDrag;
+};
+
+
+void SSkeletonTreeRow::Construct( const FArguments& InArgs, const TSharedRef<STableViewBase>& InOwnerTableView )
+{
+	Item = InArgs._Item;
+	OnDraggingItem = InArgs._OnDraggingItem;
+	FilterText = InArgs._FilterText;
+	EditableSkeleton = InArgs._EditableSkeleton;
+	SkeletonTree = InArgs._SkeletonTree;
+
+	check( Item.IsValid() );
+
+	SMultiColumnTableRow< FDisplayedTreeRowInfoPtr >::Construct( FSuperRowType::FArguments(), InOwnerTableView );
+}
+
+void SSkeletonTreeRow::ConstructChildren(ETableViewMode::Type InOwnerTableMode, const TAttribute<FMargin>& InPadding, const TSharedRef<SWidget>& InContent)
+{
+	STableRow<FDisplayedTreeRowInfoPtr>::Content = InContent;
+
+	TSharedRef<SWidget> InlineEditWidget = Item->GenerateInlineEditWidget(FilterText, FIsSelected::CreateSP(this, &STableRow::IsSelected));
+
+	// MultiColumnRows let the user decide which column should contain the expander/indenter item.
+	this->ChildSlot
+		.Padding(InPadding)
+		[
+			SNew(SVerticalBox)
+			+SVerticalBox::Slot()
+			.AutoHeight()
+			[
+				InContent
+			]
+			+SVerticalBox::Slot()
+			.AutoHeight()
+			[
+				InlineEditWidget
+			]
+		];
+}
+
+TSharedRef< SWidget > SSkeletonTreeRow::GenerateWidgetForColumn( const FName& ColumnName )
+{
+	if ( ColumnName == ColumnID_BoneLabel )
+	{
+		TSharedPtr< SHorizontalBox > RowBox;
+
+		SAssignNew( RowBox, SHorizontalBox );
+
+		RowBox->AddSlot()
+			.AutoWidth()
+			[
+				SNew( SExpanderArrow, SharedThis(this) )
+			];
+
+		Item->GenerateWidgetForNameColumn( RowBox, FilterText, FIsSelected::CreateSP(this, &STableRow::IsSelectedExclusively ) );
+
+		return RowBox.ToSharedRef();
+	}
+	else
+	{
+		return Item->GenerateWidgetForDataColumn(ColumnName);
+	}
+}
+
+void SSkeletonTreeRow::OnDragEnter( const FGeometry& MyGeometry, const FDragDropEvent& DragDropEvent )
+{
+	TSharedPtr<FSocketDragDropOp> DragConnectionOp = DragDropEvent.GetOperationAs<FSocketDragDropOp>();
+
+	// Is someone dragging a socket onto a bone?
+	if (DragConnectionOp.IsValid())
+	{
+		if ( Item->GetType() == ESkeletonTreeRowType::Bone &&
+			( *static_cast<FName*>( Item->GetData() ) != DragConnectionOp->GetSocketInfo().Socket->BoneName ) )
+		{
+			// The socket can be dropped here if we're a bone and NOT the socket's existing parent
+			DragConnectionOp->SetIcon( FEditorStyle::GetBrush( TEXT( "Graph.ConnectorFeedback.Ok" ) ) );
+		}
+		else if ( Item->GetType() == ESkeletonTreeRowType::Bone && DragConnectionOp->IsAltDrag() )
+		{
+			// For Alt-Drag, dropping onto the existing parent is fine, as we're going to copy, not move the socket
+			DragConnectionOp->SetIcon( FEditorStyle::GetBrush( TEXT( "Graph.ConnectorFeedback.Ok" ) ) );
+		}
+		else
+		{
+			DragConnectionOp->SetIcon( FEditorStyle::GetBrush( TEXT( "Graph.ConnectorFeedback.Error" ) ) );
+		}
+	}
+}
+
+void SSkeletonTreeRow::OnDragLeave( const FDragDropEvent& DragDropEvent )
+{
+	TSharedPtr<FSocketDragDropOp> DragConnectionOp = DragDropEvent.GetOperationAs<FSocketDragDropOp>();
+	if (DragConnectionOp.IsValid())
+	{
+		// Reset the drag/drop icon when leaving this row
+		DragConnectionOp->SetIcon( FEditorStyle::GetBrush( TEXT( "Graph.ConnectorFeedback.Error" ) ) );
+	}
+}
+
+FReply SSkeletonTreeRow::OnDrop( const FGeometry& MyGeometry, const FDragDropEvent& DragDropEvent )
+{
+	TSharedPtr<FSocketDragDropOp> DragConnectionOp = DragDropEvent.GetOperationAs<FSocketDragDropOp>();
+	if (DragConnectionOp.IsValid())
+	{
+		FSelectedSocketInfo SocketInfo = DragConnectionOp->GetSocketInfo();
+
+		if ( DragConnectionOp->IsAltDrag() && Item->GetType() == ESkeletonTreeRowType::Bone )
+		{
+			// In an alt-drag, the socket can be dropped on any bone
+			// (including its existing parent) to create a uniquely named copy
+			GetSkeletonTree()->DuplicateAndSelectSocket( SocketInfo, *static_cast<FName*>( Item->GetData() ) );
+		}
+		else if ( Item->GetType() == ESkeletonTreeRowType::Bone &&
+			*static_cast<FName*>( Item->GetData() ) != SocketInfo.Socket->BoneName )
+		{
+			// The socket can be dropped here if we're a bone and NOT the socket's existing parent
+			GetEditableSkeleton()->SetSocketParent(SocketInfo.Socket->SocketName, *static_cast<FName*>(Item->GetData()));
+
+			GetSkeletonTree()->CreateFromSkeleton();
+
+			return FReply::Handled();
+		}
+	}
+	else
+	{
+		if (DragDropEvent.GetOperationAs<FAssetDragDropOp>().IsValid())
+		{
+			GetSkeletonTree()->OnDropAssetToSkeletonTree( Item, DragDropEvent );
+		}
+	}
+
+	return FReply::Unhandled();
+}
+
+FReply SSkeletonTreeRow::OnDragDetected( const FGeometry& MyGeometry, const FPointerEvent& MouseEvent )
+{
+	if ( OnDraggingItem.IsBound() )
+	{
+		return OnDraggingItem.Execute( MyGeometry, MouseEvent );
+	}
+	else
+	{
+		return FReply::Unhandled();
+	}
+}
+
+int32 SSkeletonTreeRow::DoesItemHaveChildren() const
+{
+	if(Item->HasInlineEditor())
+	{
+		return 1;
+	}
+
+	return SMultiColumnTableRow<FDisplayedTreeRowInfoPtr>::DoesItemHaveChildren();
+}
+
+bool SSkeletonTreeRow::IsItemExpanded() const
+{
+	return SMultiColumnTableRow<FDisplayedTreeRowInfoPtr>::IsItemExpanded() || Item->IsInlineEditorExpanded();
+}
+
+void SSkeletonTreeRow::ToggleExpansion()
+{
+	SMultiColumnTableRow<FDisplayedTreeRowInfoPtr>::ToggleExpansion();
+	
+	if (Item->HasInlineEditor())
+	{
+		Item->ToggleInlineEditorExpansion();
+		OwnerTablePtr.Pin()->Private_SetItemExpansion(Item, Item->IsInlineEditorExpanded());
+	}
+}
+
+//////////////////////////////////////////////////////////////////////////
+// FDisplayedTreeRowInfo
+
+TSharedRef<class FEditableSkeleton> FDisplayedTreeRowInfo::GetEditableSkeleton() const
+{ 
+	return SkeletonTree.Pin()->GetEditableSkeletonInternal(); 
+}
+
+//////////////////////////////////////////////////////////////////////////
+// FDisplayedMeshBoneInfo
+
+TSharedRef<ITableRow> FDisplayedMeshBoneInfo::MakeTreeRowWidget(const TSharedRef<STableViewBase>& OwnerTable, FText FilterText )
+{
+	return
+		SNew( SSkeletonTreeRow, OwnerTable )
+		.Item( SharedThis(this) )
+		.FilterText( FilterText )
+		.SkeletonTree(SkeletonTree)
+		.EditableSkeleton(GetEditableSkeleton())
+		.OnDraggingItem( this, &FDisplayedMeshBoneInfo::OnDragDetected );
+}
+
+EVisibility FDisplayedMeshBoneInfo::GetLODIconVisibility() const
+{
+	if (bRequiredBone)
+	{
+		return EVisibility::Visible;
+	}
+
+	return EVisibility::Hidden;
+}
+
+void FDisplayedMeshBoneInfo::GenerateWidgetForNameColumn( TSharedPtr< SHorizontalBox > Box, FText& FilterText, FIsSelected InIsSelected )
+{
+	const FSlateBrush* LODIcon = FEditorStyle::GetBrush("SkeletonTree.LODBone");
+
+	Box->AddSlot()
+		.AutoWidth()
+		.Padding(FMargin(0.0f, 1.0f))
+		.VAlign(VAlign_Center)
+		.HAlign(HAlign_Center)
+		[
+			SNew(SImage)
+			.ColorAndOpacity(FSlateColor::UseForeground())
+			.Image(LODIcon)
+			.Visibility(this, &FDisplayedMeshBoneInfo::GetLODIconVisibility)
+		];
+
+	if (GetSkeletonTree()->GetPreviewScene().IsValid())
+	{
+		UDebugSkelMeshComponent* PreviewComponent = GetSkeletonTree()->GetPreviewScene()->GetPreviewMeshComponent();
+		CacheLODChange(PreviewComponent);
+	}	
+	
+	FText ToolTip = GetBoneToolTip();
+	Box->AddSlot()
+		.AutoWidth()
+		.Padding(2, 0, 0, 0)
+		.VAlign(VAlign_Center)
+		[
+			SNew( STextBlock )
+			.ColorAndOpacity(this, &FDisplayedMeshBoneInfo::GetBoneTextColor)
+			.Text( FText::FromName(BoneName) )
+			.HighlightText( FilterText )
+			.Font(this, &FDisplayedMeshBoneInfo::GetBoneTextFont)
+			.ToolTipText( ToolTip )
+		];
+}
+
+TSharedRef< SWidget > FDisplayedMeshBoneInfo::GenerateWidgetForDataColumn(const FName& DataColumnName)
+{
+	if(DataColumnName == ColumnID_RetargetingLabel)
+	{
+		return 
+			SNew(SHorizontalBox)
+			+SHorizontalBox::Slot()
+			.HAlign(HAlign_Left)
+			[
+				SAssignNew(RetargetingComboButton, SComboButton)
+				.ButtonStyle(FEditorStyle::Get(), "ToggleButton")
+				.ForegroundColor(this, &FDisplayedMeshBoneInfo::GetRetargetingComboButtonForegroundColor)
+				.ContentPadding(0)
+				.OnGetMenuContent(this, &FDisplayedMeshBoneInfo::CreateBoneTranslationRetargetingModeMenu)
+				.ToolTip(IDocumentation::Get()->CreateToolTip(
+				LOCTEXT("RetargetingToolTip", "Set bone translation retargeting mode"),
+				nullptr,
+				TEXT("Shared/Editors/Persona"),
+				TEXT("TranslationRetargeting")))
+				.ButtonContent()
+				[
+					SNew(SHorizontalBox)
+					+SHorizontalBox::Slot()
+					.VAlign(VAlign_Center)
+					[
+						SNew(STextBlock)
+						.Text(this, &FDisplayedMeshBoneInfo::GetTranslationRetargetingModeMenuTitle)
+					]
+				]
+			];
+	}
+	else if(DataColumnName == ColumnID_BlendProfileLabel)
+	{
+		bool bWritable = true;
+
+		UBlendProfile* CurrentProfile = GetSkeletonTree()->GetSelectedBlendProfile();
+
+		// We should never have this column if we don't have a profile
+		check(CurrentProfile);
+
+		return SNew(SBox)
+			.Padding(0.0f)
+			.HAlign(HAlign_Left)
+			[
+				SNew(SSpinBox<float>)
+				.Style(&FEditorStyle::GetWidgetStyle<FSpinBoxStyle>("SkeletonTree.HyperlinkSpinBox"))
+				.ContentPadding(0.0f)
+				.MinValue(0.0f)
+				.MaxValue(1000.0f)
+				.Value(CurrentProfile->GetBoneBlendScale(BoneName))
+				.OnValueCommitted(this, &FDisplayedMeshBoneInfo::OnBlendSliderCommitted)
+			];
+	}
+
+	return SNullWidget::NullWidget;
+}
+
+FSlateColor FDisplayedMeshBoneInfo::GetRetargetingComboButtonForegroundColor() const
+{
+	static const FName InvertedForegroundName("InvertedForeground");
+	static const FName DefaultForegroundName("DefaultForeground");
+
+	if (RetargetingComboButton.IsValid())
+	{
+		return RetargetingComboButton->IsHovered() ? FEditorStyle::GetSlateColor(InvertedForegroundName) : FEditorStyle::GetSlateColor(DefaultForegroundName);
+	}
+	return FSlateColor::UseForeground();
+}
+
+TSharedRef< SWidget > FDisplayedMeshBoneInfo::CreateBoneTranslationRetargetingModeMenu()
+{
+	FMenuBuilder MenuBuilder(true, nullptr);
+
+	MenuBuilder.BeginSection("BoneTranslationRetargetingMode", LOCTEXT( "BoneTranslationRetargetingModeMenuHeading", "Bone Translation Retargeting Mode" ) );
+	{
+		UEnum* const Enum = FindObject<UEnum>(ANY_PACKAGE, TEXT("EBoneTranslationRetargetingMode"), true);	
+		check(Enum);
+
+		FUIAction ActionRetargetingAnimation = FUIAction(FExecuteAction::CreateSP(this, &FDisplayedMeshBoneInfo::SetBoneTranslationRetargetingMode, EBoneTranslationRetargetingMode::Animation));
+		MenuBuilder.AddMenuEntry( Enum->GetDisplayNameText(EBoneTranslationRetargetingMode::Animation), LOCTEXT( "BoneTranslationRetargetingAnimationToolTip", "Use translation from animation." ), FSlateIcon(), ActionRetargetingAnimation);
+
+		FUIAction ActionRetargetingSkeleton = FUIAction(FExecuteAction::CreateSP(this, &FDisplayedMeshBoneInfo::SetBoneTranslationRetargetingMode, EBoneTranslationRetargetingMode::Skeleton));
+		MenuBuilder.AddMenuEntry( Enum->GetDisplayNameText(EBoneTranslationRetargetingMode::Skeleton), LOCTEXT( "BoneTranslationRetargetingSkeletonToolTip", "Use translation from Skeleton." ), FSlateIcon(), ActionRetargetingSkeleton);
+
+		FUIAction ActionRetargetingLengthScale = FUIAction(FExecuteAction::CreateSP(this, &FDisplayedMeshBoneInfo::SetBoneTranslationRetargetingMode, EBoneTranslationRetargetingMode::AnimationScaled));
+		MenuBuilder.AddMenuEntry( Enum->GetDisplayNameText(EBoneTranslationRetargetingMode::AnimationScaled), LOCTEXT( "BoneTranslationRetargetingAnimationScaledToolTip", "Use translation from animation, scale length by Skeleton's proportions." ), FSlateIcon(), ActionRetargetingLengthScale);
+
+		FUIAction ActionRetargetingAnimationRelative = FUIAction(FExecuteAction::CreateSP(this, &FDisplayedMeshBoneInfo::SetBoneTranslationRetargetingMode, EBoneTranslationRetargetingMode::AnimationRelative));
+		MenuBuilder.AddMenuEntry( Enum->GetDisplayNameText(EBoneTranslationRetargetingMode::AnimationRelative), LOCTEXT("BoneTranslationRetargetingAnimationRelativeToolTip", "Use relative translation from animation similar to an additive animation."), FSlateIcon(), ActionRetargetingAnimationRelative);
+	}
+	MenuBuilder.EndSection();
+
+	return MenuBuilder.MakeWidget();
+}
+
+FText FDisplayedMeshBoneInfo::GetTranslationRetargetingModeMenuTitle() const
+{
+	const USkeleton& Skeleton = GetEditableSkeleton()->GetSkeleton();
+
+	const int32 BoneIndex = Skeleton.GetReferenceSkeleton().FindBoneIndex( BoneName );
+	if( BoneIndex != INDEX_NONE )
+	{
+		const EBoneTranslationRetargetingMode::Type RetargetingMode = Skeleton.GetBoneTranslationRetargetingMode(BoneIndex);
+		UEnum* const Enum = FindObject<UEnum>(ANY_PACKAGE, TEXT("EBoneTranslationRetargetingMode"), true);	
+		if (Enum)
+		{
+			return Enum->GetDisplayNameText(RetargetingMode);
+		}
+	}
+
+	return LOCTEXT("None", "None");
+}
+
+void FDisplayedMeshBoneInfo::SetBoneTranslationRetargetingMode(EBoneTranslationRetargetingMode::Type NewRetargetingMode)
+{
+	GetEditableSkeleton()->SetBoneTranslationRetargetingMode(BoneName, NewRetargetingMode);
+}
+
+void FDisplayedMeshBoneInfo::SetBoneBlendProfileScale(float NewScale, bool bRecurse)
+{
+	FName BlendProfileName = GetSkeletonTree()->GetSelectedBlendProfileName();
+	GetEditableSkeleton()->SetBlendProfileScale(BlendProfileName, BoneName, NewScale, bRecurse);
+}
+
+FSlateFontInfo FDisplayedMeshBoneInfo::GetBoneTextFont() const
+{
+	if (bWeightedBone)
+	{
+		return FEditorStyle::GetWidgetStyle<FTextBlockStyle>("SkeletonTree.BoldFont").Font;
+	}
+	else
+	{
+		return FEditorStyle::GetWidgetStyle<FTextBlockStyle>("SkeletonTree.NormalFont").Font;
+	}
+}
+
+void FDisplayedMeshBoneInfo::CacheLODChange(UDebugSkelMeshComponent* PreviewComponent)
+{
+	bWeightedBone = false;
+	bRequiredBone = false;
+	if (PreviewComponent)
+	{
+		int32 BoneIndex = PreviewComponent->GetBoneIndex(BoneName);
+
+		if (BoneIndex != INDEX_NONE)
+		{
+			if (GetSkeletonTree()->IsBoneWeighted(BoneIndex, PreviewComponent))
+			{
+				//Bone is vertex weighted
+				bWeightedBone = true;
+			}
+			if (GetSkeletonTree()->IsBoneRequired(BoneIndex, PreviewComponent))
+			{
+				bRequiredBone = true;
+			}
+		}
+	}
+}
+
+FSlateColor FDisplayedMeshBoneInfo::GetBoneTextColor() const
+{
+	if (bRequiredBone)
+	{
+		return FSlateColor(FLinearColor::White);
+	}
+	else
+	{
+		return FSlateColor(FLinearColor::Gray);
+	}
+}
+
+FReply FDisplayedMeshBoneInfo::OnDragDetected(const FGeometry& MyGeometry, const FPointerEvent& MouseEvent)
+{
+	if ( MouseEvent.IsMouseButtonDown( EKeys::LeftMouseButton ) )
+	{
+		return FReply::Handled().BeginDragDrop( FBoneDragDropOp::New(GetEditableSkeleton(), BoneName ) );
+	}
+
+	return FReply::Unhandled();
+}
+
+FText FDisplayedMeshBoneInfo::GetBoneToolTip()
+{
+	bool bIsMeshBone = false;
+	bool bIsWeightedBone = false;
+	bool bMeshExists = false;
+
+	FText ToolTip;
+
+	if (GetSkeletonTree()->GetPreviewScene().IsValid())
+	{
+		UDebugSkelMeshComponent* PreviewComponent = GetSkeletonTree()->GetPreviewScene()->GetPreviewMeshComponent();
+
+		if (PreviewComponent)
+		{
+			bMeshExists = true;
+
+			int32 BoneIndex = PreviewComponent->GetBoneIndex(BoneName);
+
+			if (BoneIndex != INDEX_NONE)
+			{
+				bIsMeshBone = true;
+
+				bIsWeightedBone = GetSkeletonTree()->IsBoneWeighted(BoneIndex, PreviewComponent);
+			}
+		}
+	}
+
+	if ( !bMeshExists )
+	{
+		ToolTip = LOCTEXT( "BoneToolTipNoMeshAvailable", "This bone exists only on the skeleton as there is no current mesh set" );
+	}
+	else
+	{
+		if ( !bIsMeshBone )
+		{
+			ToolTip = LOCTEXT( "BoneToolTipSkeletonOnly", "This bone exists only on the skeleton, but not on the current mesh" );
+		}
+		else
+		{
+			if ( !bIsWeightedBone )
+			{
+				ToolTip = LOCTEXT( "BoneToolTipSkeletonAndMesh", "This bone is used by the current mesh, but has no vertices weighted against it" );
+			}
+			else
+			{
+				ToolTip = LOCTEXT( "BoneToolTipWeighted", "This bone has vertices weighted against it" );
+			}
+		}
+	}
+
+	return ToolTip;
+}
+
+void FDisplayedMeshBoneInfo::OnBlendSliderCommitted(float NewValue, ETextCommit::Type CommitType)
+{
+	if(CommitType == ETextCommit::OnEnter)
+	{
+		SetBoneBlendProfileScale(NewValue, false);
+	}
+}
+
+//////////////////////////////////////////////////////////////////////////
+// FDisplayedSocketInfo
+
+TSharedRef<ITableRow> FDisplayedSocketInfo::MakeTreeRowWidget(
+	const TSharedRef<STableViewBase>& InOwnerTable,
+	FText InFilterText )
+{
+	return
+		SNew( SSkeletonTreeRow, InOwnerTable )
+		.Item( SharedThis(this) )
+		.FilterText( InFilterText )
+		.SkeletonTree(SkeletonTree)
+		.EditableSkeleton(GetEditableSkeleton())
+		.OnDraggingItem( this, &FDisplayedSocketInfo::OnDragDetected );
+}
+
+void FDisplayedSocketInfo::GenerateWidgetForNameColumn( TSharedPtr< SHorizontalBox > Box, FText& FilterText, FIsSelected InIsSelected )
+{
+	const FSlateBrush* SocketIcon = ( ParentType == ESocketParentType::Mesh ) ?
+		FEditorStyle::GetBrush( "SkeletonTree.MeshSocket" ) :
+		FEditorStyle::GetBrush( "SkeletonTree.SkeletonSocket" );
+
+	Box->AddSlot()
+	.AutoWidth()
+	.Padding(FMargin(0.0f, 1.0f))
+	[
+		SNew( SImage )
+		.ColorAndOpacity(FSlateColor::UseForeground())
+		.Image( SocketIcon )
+	];
+
+	const FSlateFontInfo TextFont( FPaths::EngineContentDir() / TEXT("Slate/Fonts/Roboto-Regular.ttf" ), 10 );
+
+	FText ToolTip = GetSocketToolTip();
+
+	TAttribute<FText> SocketNameAttr = TAttribute<FText>::Create(TAttribute<FText>::FGetter::CreateSP(this, &FDisplayedSocketInfo::GetSocketNameAsText));
+	TSharedPtr< SInlineEditableTextBlock > InlineWidget;
+
+	Box->AddSlot()
+	.AutoWidth()
+	.Padding(2, 0, 0, 0)
+	[
+		SAssignNew( InlineWidget, SInlineEditableTextBlock )
+			.ColorAndOpacity( this, &FDisplayedSocketInfo::GetTextColor )
+			.Text( SocketNameAttr )
+			.HighlightText( FilterText )
+			.Font( TextFont )
+			.ToolTipText( ToolTip )
+			.OnVerifyTextChanged( this, &FDisplayedSocketInfo::OnVerifySocketNameChanged )
+			.OnTextCommitted( this, &FDisplayedSocketInfo::OnCommitSocketName )
+			.IsSelected( InIsSelected )
+	];
+
+	OnRenameRequested.BindSP( InlineWidget.Get(), &SInlineEditableTextBlock::EnterEditingMode);
+
+	if ( ParentType == ESocketParentType::Mesh )
+	{
+		FText SocketSuffix = IsSocketCustomized() ?
+			LOCTEXT( "CustomizedSuffix", " [Mesh]" ) :
+			LOCTEXT( "MeshSuffix", " [Mesh Only]" );
+
+		Box->AddSlot()
+		.AutoWidth()
+		[
+			SNew( STextBlock )
+			.ColorAndOpacity( FLinearColor::Gray )
+			.Text( SocketSuffix )
+			.Font(TextFont)
+			.ToolTipText( ToolTip )
+		];
+	}
+}
+
+TSharedRef< SWidget > FDisplayedSocketInfo::GenerateInlineEditWidget(FText& FilterText, FIsSelected InIsSelected)
+{
+	if (GetDefault<UPersonaOptions>()->bUseInlineSocketEditor)
+	{
+		FPropertyEditorModule& EditModule = FModuleManager::Get().GetModuleChecked<FPropertyEditorModule>("PropertyEditor");
+		FDetailsViewArgs DetailsViewArgs(/*bUpdateFromSelection=*/ false, /*bLockable=*/ false, /*bAllowSearch=*/ false, /*InNameAreaSettings=*/ FDetailsViewArgs::HideNameArea, /*bHideSelectionTip=*/ true);
+		DetailsViewArgs.bAllowFavoriteSystem = false;
+		DetailsViewArgs.bShowScrollBar = false;
+		TSharedRef<IDetailsView> DetailsView = EditModule.CreateDetailView(DetailsViewArgs);
+		DetailsView->SetObject(SocketData);
+
+		return SNew(SOverlay)
+			.Visibility_Lambda([this]() { return bInlineEditorExpanded ? EVisibility::Visible : EVisibility::Collapsed; })
+			+ SOverlay::Slot()
+			.Padding(2.0f, 4.0f, 2.0f, 4.0f)
+			[
+				DetailsView
+			]
+		+ SOverlay::Slot()
+			.HAlign(HAlign_Fill)
+			.VAlign(VAlign_Top)
+			[
+				SNew(SImage)
+				.Visibility(EVisibility::HitTestInvisible)
+			.Image(FEditorStyle::GetBrush("SkeletonTree.InlineEditorShadowTop"))
+			]
+		+ SOverlay::Slot()
+			.HAlign(HAlign_Fill)
+			.VAlign(VAlign_Bottom)
+			[
+				SNew(SImage)
+				.Visibility(EVisibility::HitTestInvisible)
+			.Image(FEditorStyle::GetBrush("SkeletonTree.InlineEditorShadowBottom"))
+			];
+	}
+	else
+	{
+		return SNullWidget::NullWidget;
+	}
+}
+
+FSlateColor FDisplayedSocketInfo::GetTextColor() const
+{
+	if (ParentType == ESocketParentType::Skeleton && bIsCustomized)
+	{
+		return FSlateColor::UseSubduedForeground();
+	}
+	else
+	{
+		return FSlateColor::UseForeground();
+	}
+}
+
+TSharedRef< SWidget > FDisplayedSocketInfo::GenerateWidgetForDataColumn(const FName& DataColumnName)
+{
+	return SNullWidget::NullWidget;
+}
+
+void FDisplayedSocketInfo::OnItemDoubleClicked()
+{
+	OnRenameRequested.ExecuteIfBound();
+}
+
+bool FDisplayedSocketInfo::CanCustomizeSocket() const
+{
+	// If the socket is on the skeleton, we have a valid mesh
+	// and there isn't one of the same name on the mesh, we can customize it
+	return (GetEditableSkeleton()->GetSkeletalMesh() && !IsSocketCustomized());
+}
+
+void FDisplayedSocketInfo::RequestRename()
+{
+	OnRenameRequested.ExecuteIfBound();
+}
+
+bool FDisplayedSocketInfo::OnVerifySocketNameChanged( const FText& InText, FText& OutErrorMessage )
+{
+	// You can't have two sockets with the same name on the mesh, nor on the skeleton,
+	// but you can have a socket with the same name on the mesh *and* the skeleton.
+	bool bVerifyName = true;
+
+	FText NewText = FText::TrimPrecedingAndTrailing(InText);
+
+	if (NewText.IsEmpty())
+	{
+		OutErrorMessage = LOCTEXT( "EmptySocketName_Error", "Sockets must have a name!");
+		bVerifyName = false;
+	}
+	else
+	{
+		bVerifyName = !GetEditableSkeleton()->DoesSocketAlreadyExist( SocketData, NewText, ParentType);
+
+		// Needs to be checked on verify.
+		if ( !bVerifyName )
+		{
+
+			// Tell the user that the socket is a duplicate
+			OutErrorMessage = LOCTEXT( "DuplicateSocket_Error", "Socket name in use!");
+		}
+	}
+
+	return bVerifyName;
+}
+
+void FDisplayedSocketInfo::OnCommitSocketName( const FText& InText, ETextCommit::Type CommitInfo )
+{
+	FText NewText = FText::TrimPrecedingAndTrailing(InText);
+
+	// Notify skeleton tree of socket rename
+	GetEditableSkeleton()->RenameSocket(SocketData->SocketName, FName(*NewText.ToString()));
+}
+
+FReply FDisplayedSocketInfo::OnDragDetected(const FGeometry& MyGeometry, const FPointerEvent& MouseEvent)
+{
+	if ( MouseEvent.IsMouseButtonDown( EKeys::LeftMouseButton ) )
+	{
+		FSelectedSocketInfo SocketInfo( SocketData, ParentType == ESocketParentType::Skeleton );
+
+		return FReply::Handled().BeginDragDrop( FSocketDragDropOp::New( SocketInfo, MouseEvent.IsAltDown() ) );
+	}
+
+	return FReply::Unhandled();
+}
+
+FText FDisplayedSocketInfo::GetSocketToolTip()
+{
+	FText ToolTip;
+
+	if ( ParentType == ESocketParentType::Skeleton && bIsCustomized == false )
+	{
+		ToolTip = LOCTEXT( "SocketToolTipSkeletonOnly", "This socket is on the skeleton only. It is shared with all meshes that use this skeleton" );
+	}
+	else if ( ParentType == ESocketParentType::Mesh && bIsCustomized == false )
+	{
+		ToolTip = LOCTEXT( "SocketToolTipMeshOnly", "This socket is on the current mesh only" );
+	}
+	else if ( ParentType == ESocketParentType::Skeleton )
+	{
+		ToolTip = LOCTEXT( "SocketToolTipSkeleton", "This socket is on the skeleton (shared with all meshes that use the skeleton), and the current mesh has duplciated version of it" );
+	}
+	else
+	{
+		ToolTip = LOCTEXT( "SocketToolTipCustomized", "This socket is on the current mesh, customizing the socket of the same name on the skeleton" );
+	}
+
+	return ToolTip;
+}
+
+//////////////////////////////////////////////////////////////////////////
+// FDisplayedAttachedAssetInfo
+
+TSharedRef<FDisplayedAttachedAssetInfo> FDisplayedAttachedAssetInfo::Make(const FName& InAttachedTo, UObject* InAsset, TWeakPtr<SSkeletonTree> InSkeletonTree)
+{
+	FDisplayedAttachedAssetInfo* Info = new FDisplayedAttachedAssetInfo(InAttachedTo, InAsset);
+	Info->SkeletonTree = InSkeletonTree;
+
+	Info->AssetComponent = PersonaUtils::GetComponentForAttachedObject(InSkeletonTree.Pin()->GetPreviewScene()->GetPreviewMeshComponent(), InAsset, InAttachedTo);
+
+	return MakeShareable(Info);
+}
+
+TSharedRef<ITableRow> FDisplayedAttachedAssetInfo::MakeTreeRowWidget(
+	const TSharedRef<STableViewBase>& InOwnerTable,
+	FText InFilterText )
+{
+	return
+		SNew(SSkeletonTreeRow, InOwnerTable)
+		.Item(SharedThis(this))
+		.FilterText(InFilterText)
+		.SkeletonTree(SkeletonTree)
+		.EditableSkeleton(GetEditableSkeleton());
+}
+
+void FDisplayedAttachedAssetInfo::GenerateWidgetForNameColumn( TSharedPtr< SHorizontalBox > Box, FText& FilterText, FIsSelected InIsSelected )
+{
+	UActorFactory* ActorFactory = FActorFactoryAssetProxy::GetFactoryForAssetObject( Asset );
+	const FSlateBrush* IconBrush = FSlateIconFinder::FindIconBrushForClass(ActorFactory->GetDefaultActorClass(FAssetData()));
+	
+	Box->AddSlot()
+		.Padding(FMargin(0.0f, 1.0f))
+		.AutoWidth()
+		[
+			SNew( SImage )
+			.Image( IconBrush )
+		];
+
+	const FSlateFontInfo TextFont( FPaths::EngineContentDir() / TEXT("Slate/Fonts/Roboto-Regular.ttf"), 10 );
+	const FLinearColor TextColor(FLinearColor::White);
+
+	Box->AddSlot()
+		.AutoWidth()
+		[
+			SNew( STextBlock )
+			.ColorAndOpacity(TextColor)
+			.Text( FText::FromString(Asset->GetName()) )
+			.HighlightText( FilterText )
+			.Font(TextFont)
+		];
+
+	Box->AddSlot()
+		.AutoWidth()
+		.Padding(5.0f,0.0f)
+		[
+			SNew( STextBlock )
+			.ColorAndOpacity( FLinearColor::Gray )
+			.Text(LOCTEXT( "AttachedAssetPreviewText", "[Preview Only]" ) )
+			.Font(TextFont)
+			.ToolTipText( LOCTEXT( "AttachedAssetPreviewText_ToolTip", "Attached assets in Persona are preview only and do not carry through to the game." ) )
+		];
+}
+
+TSharedRef< SWidget > FDisplayedAttachedAssetInfo::GenerateWidgetForDataColumn(const FName& DataColumnName)
+{
+	if(DataColumnName == ColumnID_RetargetingLabel)
+	{
+		return SNew(SHorizontalBox)
+			+ SHorizontalBox::Slot()
+			.AutoWidth()
+			.HAlign(HAlign_Left)
+			[
+				SNew(SCheckBox)
+				.ToolTipText(LOCTEXT("TranslationCheckBoxToolTip", "Click to toggle visibility of this asset"))
+				.OnCheckStateChanged(this, &FDisplayedAttachedAssetInfo::OnToggleAssetDisplayed)
+				.IsChecked(this, &FDisplayedAttachedAssetInfo::IsAssetDisplayed)
+				.Style(FEditorStyle::Get(), "CheckboxLookToggleButtonCheckbox")
+				[
+					SNew(SImage)
+					.Image(this, &FDisplayedAttachedAssetInfo::OnGetAssetDisplayedButtonImage)
+					.ColorAndOpacity(FLinearColor::Black)
+				]
+			];
+	}
+
+	return SNullWidget::NullWidget;
+}
+
+ECheckBoxState FDisplayedAttachedAssetInfo::IsAssetDisplayed() const
+{
+	if(AssetComponent.IsValid())
+	{
+		return AssetComponent->IsVisible() ? ECheckBoxState::Checked : ECheckBoxState::Unchecked;
+	}
+	return ECheckBoxState::Undetermined;
+}
+
+void FDisplayedAttachedAssetInfo::OnToggleAssetDisplayed( ECheckBoxState InCheckboxState )
+{
+	if(AssetComponent.IsValid())
+	{
+		AssetComponent->SetVisibility(InCheckboxState == ECheckBoxState::Checked);
+	}
+}
+
+const FSlateBrush* FDisplayedAttachedAssetInfo::OnGetAssetDisplayedButtonImage() const
+{
+	return IsAssetDisplayed() == ECheckBoxState::Checked ?
+		FEditorStyle::GetBrush( "Kismet.VariableList.ExposeForInstance" ) :
+		FEditorStyle::GetBrush( "Kismet.VariableList.HideForInstance" );
+}
+
+void FDisplayedAttachedAssetInfo::OnItemDoubleClicked()
+{
+	TArray<UObject*> AssetsToSync;
+	AssetsToSync.Add(Asset);
+
+	FContentBrowserModule& ContentBrowserModule = FModuleManager::Get().LoadModuleChecked<FContentBrowserModule>("ContentBrowser");
+	ContentBrowserModule.Get().SyncBrowserToAssets( AssetsToSync );
+}
+
+//////////////////////////////////////////////////////////////////////////
+// SSkeletonTree
+
+void SSkeletonTree::Construct(const FArguments& InArgs, const TSharedRef<FEditableSkeleton>& InEditableSkeleton, const FSkeletonTreeArgs& InSkeletonTreeArgs)
+{
+	BoneFilter = EBoneFilter::All;
+	SocketFilter = ESocketFilter::Active;
+	bShowingAdvancedOptions = false;
+	bSelectingSocket = false;
+	bSelectingBone = false;
+	bDeselectingAll = false;
+
+	EditableSkeleton = InEditableSkeleton;
+	PreviewScene = InSkeletonTreeArgs.PreviewScene;
+	IsEditable = InArgs._IsEditable;
+
+	SetPreviewComponentSocketFilter();
+
+	// Register delegates
+	InSkeletonTreeArgs.OnPostUndo.Add(FSimpleDelegate::CreateSP(this, &SSkeletonTree::PostUndo));
+
+	PreviewScene->RegisterOnLODChanged(FSimpleDelegate::CreateSP(this, &SSkeletonTree::OnLODSwitched));
+
+	InEditableSkeleton->RegisterOnSkeletonHierarchyChanged(USkeleton::FOnSkeletonHierarchyChanged::CreateSP(this, &SSkeletonTree::PostUndo));
+
+	RegisterOnObjectSelected(InSkeletonTreeArgs.OnObjectSelected);
+
+	// Register and bind all our menu commands
+	FSkeletonTreeCommands::Register();
+	BindCommands();
+
+	this->ChildSlot
+	[
+		SNew( SVerticalBox )
+		
+		+ SVerticalBox::Slot()
+		.AutoHeight()
+		.Padding( FMargin( 0.0f, 0.0f, 0.0f, 4.0f ) )
+		[
+			SAssignNew( NameFilterBox, SSearchBox )
+			.SelectAllTextWhenFocused( true )
+			.OnTextChanged( this, &SSkeletonTree::OnFilterTextChanged )
+			.HintText( LOCTEXT( "SearchBoxHint", "Search Skeleton Tree...") )
+			.AddMetaData<FTagMetaData>(TEXT("SkelTree.Search"))
+		]
+
+		+ SVerticalBox::Slot()
+		.Padding( FMargin( 0.0f, 4.0f, 0.0f, 0.0f ) )
+		[
+			SAssignNew(TreeHolder, SOverlay)
+		]
+
+		+ SVerticalBox::Slot()
+		.AutoHeight()
+		[
+			SNew( SHorizontalBox )
+
+			+ SHorizontalBox::Slot()
+			.Padding(0.0f, 0.0f, 2.0f, 0.0f)
+			.VAlign(VAlign_Center)
+			.HAlign(HAlign_Left)
+			.FillWidth(1.0f)
+			[
+				SNew(SHorizontalBox)
+				
+				+ SHorizontalBox::Slot()
+				.Padding(0.0f, 0.0f, 2.0f, 0.0f)
+				.VAlign(VAlign_Center)
+				[
+					SAssignNew(BlendProfilePicker, SBlendProfilePicker, GetEditableSkeleton())
+					.Standalone(true)
+					.OnBlendProfileSelected(this, &SSkeletonTree::OnBlendProfileSelected)
+				]
+			]
+
+			+ SHorizontalBox::Slot()
+			.Padding( 0.0f, 0.0f, 2.0f, 0.0f )
+			.VAlign(VAlign_Center)
+			.HAlign(HAlign_Right)
+			.AutoWidth()
+			[
+				SAssignNew( FilterComboButton, SComboButton )
+				.ButtonStyle(FEditorStyle::Get(), "ToggleButton")
+				.ForegroundColor(this, &SSkeletonTree::GetFilterComboButtonForegroundColor)
+				.ContentPadding(0.0f)
+				.OnGetMenuContent( this, &SSkeletonTree::CreateFilterMenu )
+				.ToolTipText( LOCTEXT( "BoneFilterToolTip", "Change which types of bones and sockets are shown, retargeting options" ) )
+				.AddMetaData<FTagMetaData>(TEXT("SkelTree.Bones"))
+				.ButtonContent()
+				[
+					SNew(SHorizontalBox)
+					+SHorizontalBox::Slot()
+					.AutoWidth()
+					.VAlign(VAlign_Center)
+					[
+						SNew(SImage).Image( FEditorStyle::GetBrush("GenericViewButton") )
+					]
+					+SHorizontalBox::Slot()
+					.AutoWidth()
+					.Padding(2, 0, 0, 0)
+					.VAlign(VAlign_Center)
+					[
+						SNew( STextBlock )
+						.Text( this, &SSkeletonTree::GetFilterMenuTitle )
+					]
+				]
+			]
+		]
+	];
+
+	CreateTreeColumns();
+
+	OnLODSwitched();
+}
+
+SSkeletonTree::~SSkeletonTree()
+{
+	if (EditableSkeleton.IsValid())
+	{
+		EditableSkeleton.Pin()->UnregisterOnSkeletonHierarchyChanged(this);
+	}
+}
+
+void SSkeletonTree::BindCommands()
+{
+	// This should not be called twice on the same instance
+	check( !UICommandList.IsValid() );
+
+	UICommandList = MakeShareable( new FUICommandList );
+
+	FUICommandList& CommandList = *UICommandList;
+
+	// Grab the list of menu commands to bind...
+	const FSkeletonTreeCommands& MenuActions = FSkeletonTreeCommands::Get();
+
+	// ...and bind them all
+
+	// Bone Filter commands
+	CommandList.MapAction(
+		MenuActions.ShowAllBones,
+		FExecuteAction::CreateSP( this, &SSkeletonTree::SetBoneFilter, EBoneFilter::All ),
+		FCanExecuteAction(),
+		FIsActionChecked::CreateSP( this, &SSkeletonTree::IsBoneFilter, EBoneFilter::All ) );
+
+	CommandList.MapAction(
+		MenuActions.ShowMeshBones,
+		FExecuteAction::CreateSP( this, &SSkeletonTree::SetBoneFilter, EBoneFilter::Mesh ),
+		FCanExecuteAction(),
+		FIsActionChecked::CreateSP( this, &SSkeletonTree::IsBoneFilter, EBoneFilter::Mesh ) );
+
+	CommandList.MapAction(
+		MenuActions.ShowLODBones,
+		FExecuteAction::CreateSP(this, &SSkeletonTree::SetBoneFilter, EBoneFilter::LOD),
+		FCanExecuteAction(),
+		FIsActionChecked::CreateSP(this, &SSkeletonTree::IsBoneFilter, EBoneFilter::LOD));
+	
+	CommandList.MapAction(
+		MenuActions.ShowWeightedBones,
+		FExecuteAction::CreateSP( this, &SSkeletonTree::SetBoneFilter, EBoneFilter::Weighted ),
+		FCanExecuteAction(),
+		FIsActionChecked::CreateSP( this, &SSkeletonTree::IsBoneFilter, EBoneFilter::Weighted ) );
+
+	CommandList.MapAction(
+		MenuActions.HideBones,
+		FExecuteAction::CreateSP( this, &SSkeletonTree::SetBoneFilter, EBoneFilter::None ),
+		FCanExecuteAction(),
+		FIsActionChecked::CreateSP( this, &SSkeletonTree::IsBoneFilter, EBoneFilter::None ) );
+
+	// Socket filter commands
+	CommandList.MapAction(
+		MenuActions.ShowActiveSockets,
+		FExecuteAction::CreateSP( this, &SSkeletonTree::SetSocketFilter, ESocketFilter::Active ),
+		FCanExecuteAction(),
+		FIsActionChecked::CreateSP( this, &SSkeletonTree::IsSocketFilter, ESocketFilter::Active ) );
+
+	CommandList.MapAction(
+		MenuActions.ShowMeshSockets,
+		FExecuteAction::CreateSP( this, &SSkeletonTree::SetSocketFilter, ESocketFilter::Mesh ),
+		FCanExecuteAction(),
+		FIsActionChecked::CreateSP( this, &SSkeletonTree::IsSocketFilter, ESocketFilter::Mesh ) );
+
+	CommandList.MapAction(
+		MenuActions.ShowSkeletonSockets,
+		FExecuteAction::CreateSP( this, &SSkeletonTree::SetSocketFilter, ESocketFilter::Skeleton ),
+		FCanExecuteAction(),
+		FIsActionChecked::CreateSP( this, &SSkeletonTree::IsSocketFilter, ESocketFilter::Skeleton ) );
+
+	CommandList.MapAction(
+		MenuActions.ShowAllSockets,
+		FExecuteAction::CreateSP( this, &SSkeletonTree::SetSocketFilter, ESocketFilter::All ),
+		FCanExecuteAction(),
+		FIsActionChecked::CreateSP( this, &SSkeletonTree::IsSocketFilter, ESocketFilter::All ) );
+
+	CommandList.MapAction(
+		MenuActions.HideSockets,
+		FExecuteAction::CreateSP( this, &SSkeletonTree::SetSocketFilter, ESocketFilter::None ),
+		FCanExecuteAction(),
+		FIsActionChecked::CreateSP( this, &SSkeletonTree::IsSocketFilter, ESocketFilter::None ) );
+
+	CommandList.MapAction(
+		MenuActions.ShowRetargeting,
+		FExecuteAction::CreateSP(this, &SSkeletonTree::OnChangeShowingAdvancedOptions),
+		FCanExecuteAction(),
+		FIsActionChecked::CreateSP(this, &SSkeletonTree::IsShowingAdvancedOptions));
+
+	// Socket manipulation commands
+	CommandList.MapAction(
+		MenuActions.AddSocket,
+		FExecuteAction::CreateSP( this, &SSkeletonTree::OnAddSocket ),
+		FCanExecuteAction::CreateSP( this, &SSkeletonTree::IsAddingSocketsAllowed ) );
+
+	CommandList.MapAction(
+		FGenericCommands::Get().Rename,
+		FExecuteAction::CreateSP( this, &SSkeletonTree::OnRenameSocket ),
+		FCanExecuteAction::CreateSP( this, &SSkeletonTree::CanRenameSelected ) );
+
+	CommandList.MapAction(
+		MenuActions.CreateMeshSocket,
+		FExecuteAction::CreateSP( this, &SSkeletonTree::OnCustomizeSocket ) );
+
+	CommandList.MapAction(
+		MenuActions.RemoveMeshSocket,
+		FExecuteAction::CreateSP( this, &SSkeletonTree::OnDeleteSelectedRows ) ); // Removing customization just deletes the mesh socket
+
+	CommandList.MapAction(
+		MenuActions.PromoteSocketToSkeleton,
+		FExecuteAction::CreateSP( this, &SSkeletonTree::OnPromoteSocket ) ); // Adding customization just deletes the mesh socket
+
+	CommandList.MapAction(
+		MenuActions.DeleteSelectedRows,
+		FExecuteAction::CreateSP( this, &SSkeletonTree::OnDeleteSelectedRows ) );
+
+	CommandList.MapAction(
+		MenuActions.CopyBoneNames,
+		FExecuteAction::CreateSP( this, &SSkeletonTree::OnCopyBoneNames ) );
+
+	CommandList.MapAction(
+		MenuActions.ResetBoneTransforms,
+		FExecuteAction::CreateSP(this, &SSkeletonTree::OnResetBoneTransforms ) );
+
+	CommandList.MapAction(
+		MenuActions.CopySockets,
+		FExecuteAction::CreateSP( this, &SSkeletonTree::OnCopySockets ) );
+
+	CommandList.MapAction(
+		MenuActions.PasteSockets,
+		FExecuteAction::CreateSP( this, &SSkeletonTree::OnPasteSockets, false ) );
+
+	CommandList.MapAction(
+		MenuActions.PasteSocketsToSelectedBone,
+		FExecuteAction::CreateSP(this, &SSkeletonTree::OnPasteSockets, true));
+}
+
+TSharedRef<ITableRow> SSkeletonTree::MakeTreeRowWidget(TSharedPtr<FDisplayedTreeRowInfo> InInfo, const TSharedRef<STableViewBase>& OwnerTable)
+{
+	check( InInfo.IsValid() );
+	
+	return InInfo->MakeTreeRowWidget( OwnerTable, FilterText );
+}
+
+void SSkeletonTree::GetChildrenForInfo(TSharedPtr<FDisplayedTreeRowInfo> InInfo, TArray< TSharedPtr<FDisplayedTreeRowInfo> >& OutChildren)
+{
+	check(InInfo.IsValid());
+	OutChildren = InInfo->Children;
+}
+
+bool SSkeletonTree::AttachToParent( TSharedRef<FDisplayedTreeRowInfo> ItemToAttach, FName ParentName, int32 ItemsToInclude)
+{
+	// Find the parent info
+	for (int32 BoneIndex = 0; BoneIndex < DisplayMirror.Num(); ++BoneIndex)
+	{
+		TSharedRef<FDisplayedTreeRowInfo> CurrentItem = DisplayMirror[BoneIndex];
+		
+		// does the item match our filter
+		if ( (ItemsToInclude & (int32)CurrentItem->GetType()) != 0 )
+		{
+			if ( CurrentItem->GetRowItemName() == ParentName )
+			{
+				CurrentItem->Children.Insert( ItemToAttach, 0 );
+				return true;
+			}
+		}
+	}
+	return false;
+}
+
+void SSkeletonTree::CreateTreeColumns()
+{
+	TSharedRef<SHeaderRow> TreeHeaderRow = SNew(SHeaderRow)
+		+ SHeaderRow::Column(ColumnID_BoneLabel)
+		.DefaultLabel(LOCTEXT("SkeletonBoneNameLabel", "Name"))
+		.FillWidth(0.5f);
+
+	if (bShowingAdvancedOptions)
+	{
+		TreeHeaderRow->AddColumn(
+			SHeaderRow::Column(ColumnID_RetargetingLabel)
+			.DefaultLabel(LOCTEXT("SkeletonBoneTranslationRetargetingLabel", "Translation Retargeting"))
+			.FillWidth(0.25f)
+			);
+	}
+
+	if(BlendProfilePicker->GetSelectedBlendProfileName() != NAME_None)
+	{
+		TreeHeaderRow->AddColumn(
+			SHeaderRow::Column(ColumnID_BlendProfileLabel)
+			.DefaultLabel(LOCTEXT("BlendProfileLabel", "Blend Profile Scale"))
+			.FillWidth(0.25f));
+	}
+
+	SAssignNew(SkeletonTreeView, SMeshSkeletonTreeRowType)
+		.TreeItemsSource(&SkeletonRowList)
+		.OnGenerateRow(this, &SSkeletonTree::MakeTreeRowWidget)
+		.OnGetChildren(this, &SSkeletonTree::GetChildrenForInfo)
+		.OnContextMenuOpening(this, &SSkeletonTree::CreateContextMenu)
+		.OnSelectionChanged(this, &SSkeletonTree::OnSelectionChanged)
+		.OnItemScrolledIntoView(this, &SSkeletonTree::OnItemScrolledIntoView)
+		.OnMouseButtonDoubleClick(this, &SSkeletonTree::OnTreeDoubleClick)
+		.OnSetExpansionRecursive(this, &SSkeletonTree::SetTreeItemExpansionRecursive)
+		.ItemHeight(24)
+		.HeaderRow
+		(
+			TreeHeaderRow
+		);
+
+	TreeHolder->ClearChildren();
+	TreeHolder->AddSlot()
+	[
+		SNew(SScrollBorder, SkeletonTreeView.ToSharedRef())
+		[
+			SkeletonTreeView.ToSharedRef()
+		]
+	];
+
+	CreateFromSkeleton();
+}
+
+
+void SSkeletonTree::CreateFromSkeleton(USkeletalMeshSocket* SocketToRename /*= nullptr*/)
+{
+	const USkeleton& Skeleton = GetEditableSkeletonInternal()->GetSkeleton();
+	
+	SkeletonRowList.Empty();
+
+	const FReferenceSkeleton& RefSkeleton = Skeleton.GetReferenceSkeleton();
+
+	DisplayMirror.Empty( RefSkeleton.GetNum() );
+
+	if( BoneFilter != EBoneFilter::None )
+	{
+		for (int32 BoneIndex = 0; BoneIndex < RefSkeleton.GetNum(); ++BoneIndex)
+		{
+			const FName& BoneName = RefSkeleton.GetBoneName(BoneIndex);
+			if ( !FilterText.IsEmpty() && !BoneName.ToString().Contains( FilterText.ToString()) )
+			{
+				continue;
+			}
+		
+			if (GetPreviewScene().IsValid())
+			{
+				UDebugSkelMeshComponent* PreviewMeshComponent = GetPreviewScene()->GetPreviewMeshComponent();
+				if (PreviewMeshComponent)
+				{
+					int32 BoneMeshIndex = PreviewMeshComponent->GetBoneIndex(BoneName);
+
+					// Remove non-mesh bones if we're filtering
+					if ((BoneFilter == EBoneFilter::Mesh || BoneFilter == EBoneFilter::Weighted || BoneFilter == EBoneFilter::LOD) &&
+						BoneMeshIndex == INDEX_NONE)
+					{
+						continue;
+					}
+
+					// Remove non-vertex-weighted bones if we're filtering
+					if (BoneFilter == EBoneFilter::Weighted && !IsBoneWeighted(BoneMeshIndex, PreviewMeshComponent))
+					{
+						continue;
+					}
+
+					// Remove non-vertex-weighted bones if we're filtering
+					if (BoneFilter == EBoneFilter::LOD && !IsBoneRequired(BoneMeshIndex, PreviewMeshComponent))
+					{
+						continue;
+					}
+				}
+			}
+
+			int32 ParentIndex = RefSkeleton.GetParentIndex(BoneIndex);
+
+			TSharedRef<FDisplayedMeshBoneInfo> DisplayBone = FDisplayedMeshBoneInfo::Make( RefSkeleton.GetBoneName(BoneIndex), SharedThis(this) );
+		
+			if ( BoneIndex > 0 && FilterText.IsEmpty() && DisplayMirror.Num() != 0 )	// No hierarchy when filtertext is non-empty
+			{
+				check(ParentIndex < BoneIndex);
+
+				// We support filtering the list, so ParentIndex isn't necessarily correct in the DisplayMirror any more, so we need to search for it by name
+				FName ParentName = RefSkeleton.GetBoneName(ParentIndex);
+				bool bFoundRemappedParentIndex = false;
+
+				for ( int32 i = 0; i < DisplayMirror.Num(); ++i )
+				{
+					FDisplayedTreeRowInfoPtr TreeRowInfo = DisplayMirror[i];
+
+					// At this point, we can assume that *all* of DisplayMirror contains bones, not sockets
+					check( TreeRowInfo->GetType() == ESkeletonTreeRowType::Bone );
+
+					if ( *static_cast<FName*>( TreeRowInfo->GetData() ) == ParentName )
+					{
+						ParentIndex = i;
+						bFoundRemappedParentIndex = true;
+						break;
+					}
+				}
+
+				if ( bFoundRemappedParentIndex )
+				{
+					DisplayMirror[ParentIndex]->Children.Add(DisplayBone);
+				}
+				else
+				{
+					// The parent bone didn't pass the filter, so just add this bone to the base of the tree
+					SkeletonRowList.Add(DisplayBone);
+				}
+			}
+ 			else
+			{
+				SkeletonRowList.Add(DisplayBone);
+			}
+
+			DisplayMirror.Add(DisplayBone);
+			SkeletonTreeView->SetItemExpansion(DisplayBone, true);
+		}
+	}
+
+	// Add the sockets for the skeleton
+	if ( SocketFilter == ESocketFilter::Active || SocketFilter == ESocketFilter::All || SocketFilter == ESocketFilter::Skeleton )
+	{
+		AddSocketsFromData( Skeleton.Sockets, ESocketParentType::Skeleton, SocketToRename );
+	}
+	
+	if ( SocketFilter == ESocketFilter::Active || SocketFilter == ESocketFilter::All || SocketFilter == ESocketFilter::Mesh )
+	{
+		// Add the sockets for the mesh
+		if ( USkeletalMesh* const SkeletalMesh = GetEditableSkeletonInternal()->GetSkeletalMesh() )
+		{
+			AddSocketsFromData( SkeletalMesh->GetMeshOnlySocketList(), ESocketParentType::Mesh, SocketToRename );
+		}
+	}
+
+	//Add the attached mesh items last, these are the most child like of all the items that can go in the skeleton tree
+
+	// Mesh attached items...
+	if ( USkeletalMesh* const SkeletalMesh = GetEditableSkeletonInternal()->GetSkeletalMesh() )
+	{
+		AddAttachedAssets( SkeletalMesh->PreviewAttachedAssetContainer );
+	}
+
+	// ...skeleton attached items
+	AddAttachedAssets( Skeleton.PreviewAttachedAssetContainer );
+
+	SkeletonTreeView->RequestTreeRefresh();
+}
+
+void SSkeletonTree::AddSocketsFromData(const TArray< USkeletalMeshSocket* >& SocketArray, ESocketParentType ParentType, USkeletalMeshSocket* SocketToRename)
+{
+	const USkeleton& Skeleton = GetEditableSkeletonInternal()->GetSkeleton();
+
+	for (auto SocketIt = SocketArray.CreateConstIterator(); SocketIt; ++SocketIt)
+	{
+		USkeletalMeshSocket* Socket = *(SocketIt);
+
+		if (!FilterText.IsEmpty() && !Socket->SocketName.ToString().Contains(FilterText.ToString()))
+		{
+			continue;
+		}
+
+		bool bIsCustomized = false;
+
+		if (ParentType == ESocketParentType::Mesh)
+		{
+			bIsCustomized = GetEditableSkeletonInternal()->DoesSocketAlreadyExist(NULL, FText::FromName(Socket->SocketName), ESocketParentType::Skeleton);
+		}
+		else
+		{
+			if (USkeletalMesh* const Mesh = GetEditableSkeletonInternal()->GetSkeletalMesh())
+			{
+				bIsCustomized = GetEditableSkeletonInternal()->DoesSocketAlreadyExist(NULL, FText::FromName(Socket->SocketName), ESocketParentType::Mesh);
+
+				if (SocketFilter == ESocketFilter::Active && bIsCustomized)
+				{
+					// Don't add the skeleton socket if it's already added for the mesh
+					continue;
+				}
+			}
+		}
+
+		TSharedRef<FDisplayedSocketInfo> DisplaySocket = FDisplayedSocketInfo::Make(Socket, ParentType, SharedThis(this), bIsCustomized);
+
+		if (Socket == SocketToRename)
+		{
+			SkeletonTreeView->SetSelection(DisplaySocket);
+			OnRenameSocket();
+		}
+		DisplayMirror.Add(DisplaySocket);
+
+		if (!AttachToParent(DisplaySocket, Socket->BoneName, (int32)ESkeletonTreeRowType::Bone))
+		{
+			// Just add it to the list if the parent bone isn't currently displayed
+			SkeletonRowList.Add(DisplaySocket);
+		}
+
+		SkeletonTreeView->SetItemExpansion(DisplaySocket, true);
+	}
+}
+
+class FBoneTreeSelection
+{
+public:
+	TArray<TSharedPtr<FDisplayedTreeRowInfo>> SelectedItems;
+
+	TArray<TSharedPtr<FDisplayedMeshBoneInfo>> SelectedBones;
+	TArray<TSharedPtr<FDisplayedSocketInfo>> SelectedSockets;
+	TArray<TSharedPtr<FDisplayedAttachedAssetInfo>> SelectedAssets;
+
+	FBoneTreeSelection(TArray<TSharedPtr<FDisplayedTreeRowInfo>> InSelectedItems) : SelectedItems(InSelectedItems)
+	{
+		for ( auto ItemIt = SelectedItems.CreateConstIterator(); ItemIt; ++ItemIt )
+		{
+			FDisplayedTreeRowInfoPtr Item = *(ItemIt);
+
+			switch(Item->GetType())
+			{
+				case ESkeletonTreeRowType::Bone:
+				{
+					SelectedBones.Add( StaticCastSharedPtr< FDisplayedMeshBoneInfo >(Item) );
+					break;
+				}
+				case ESkeletonTreeRowType::Socket:
+				{
+					SelectedSockets.Add( StaticCastSharedPtr< FDisplayedSocketInfo >(Item) );
+					break;
+				}
+				case ESkeletonTreeRowType::AttachedAsset:
+				{
+					SelectedAssets.Add( StaticCastSharedPtr< FDisplayedAttachedAssetInfo >(Item) );
+					break;
+				}
+				default:
+				{
+					check(false); // Unknown row type!
+				}
+			}
+		}
+	}
+
+	bool IsMultipleItemsSelected() const
+	{
+		return SelectedItems.Num() > 1;
+	}
+
+	bool IsSingleItemSelected() const
+	{
+		return SelectedItems.Num() == 1;
+	}
+
+	bool IsSingleOfTypeSelected(ESkeletonTreeRowType ItemType) const
+	{
+		if(IsSingleItemSelected())
+		{
+			switch (ItemType)
+			{
+				case ESkeletonTreeRowType::Bone:
+				{
+					return SelectedBones.Num() == 1;
+				}
+				case ESkeletonTreeRowType::Socket:
+				{
+					return SelectedSockets.Num() == 1;
+				}
+				case ESkeletonTreeRowType::AttachedAsset:
+				{
+					return SelectedAssets.Num() == 1;
+				}
+				default:
+				{
+					check(false); // Unknown type
+				}
+			}
+		}
+		return false;
+	}
+
+	TSharedPtr<FDisplayedTreeRowInfo> GetSingleSelectedItem()
+	{
+		check(IsSingleItemSelected());
+		return SelectedItems[0];
+	}
+
+	bool HasSelectedOfType(ESkeletonTreeRowType ItemType) const
+	{
+		switch (ItemType)
+		{
+			case ESkeletonTreeRowType::Bone:
+			{
+				return SelectedBones.Num() > 0;
+			}
+			case ESkeletonTreeRowType::Socket:
+			{
+				return SelectedSockets.Num() > 0;
+			}
+			case ESkeletonTreeRowType::AttachedAsset:
+			{
+				return SelectedAssets.Num() > 0;
+			}
+			default:
+			{
+				check(false); // Unknown type
+			}
+		}
+		return false;
+	}
+};
+
+TSharedPtr< SWidget > SSkeletonTree::CreateContextMenu()
+{
+	const FSkeletonTreeCommands& Actions = FSkeletonTreeCommands::Get();
+
+	FBoneTreeSelection BoneTreeSelection(SkeletonTreeView->GetSelectedItems());
+
+	const bool CloseAfterSelection = true;
+	FMenuBuilder MenuBuilder( CloseAfterSelection, UICommandList );
+	{
+		if(BoneTreeSelection.HasSelectedOfType(ESkeletonTreeRowType::AttachedAsset) || BoneTreeSelection.HasSelectedOfType(ESkeletonTreeRowType::Socket))
+		{
+			MenuBuilder.BeginSection("SkeletonTreeSelectedItemsActions", LOCTEXT( "SelectedActions", "Selected Item Actions" ) );
+			MenuBuilder.AddMenuEntry( Actions.DeleteSelectedRows );
+			MenuBuilder.EndSection();
+		}
+
+		if(BoneTreeSelection.HasSelectedOfType(ESkeletonTreeRowType::Bone))
+		{
+			MenuBuilder.BeginSection("SkeletonTreeBonesAction", LOCTEXT( "BoneActions", "Selected Bone Actions" ) );
+			MenuBuilder.AddMenuEntry( Actions.CopyBoneNames );
+			MenuBuilder.AddMenuEntry( Actions.ResetBoneTransforms );
+
+			if(BoneTreeSelection.IsSingleOfTypeSelected(ESkeletonTreeRowType::Bone))
+			{
+				MenuBuilder.AddMenuEntry( Actions.AddSocket );
+				MenuBuilder.AddMenuEntry( Actions.PasteSockets );
+				MenuBuilder.AddMenuEntry( Actions.PasteSocketsToSelectedBone );
+			}
+
+			MenuBuilder.EndSection();
+
+			UBlendProfile* const SelectedBlendProfile = BlendProfilePicker->GetSelectedBlendProfile();
+			if(SelectedBlendProfile && BoneTreeSelection.IsSingleOfTypeSelected(ESkeletonTreeRowType::Bone))
+			{
+				TSharedPtr<FDisplayedMeshBoneInfo> BoneInfo = BoneTreeSelection.SelectedBones[0];
+
+				FName BoneName = *static_cast<FName*>(BoneInfo->GetData());
+				const USkeleton& Skeleton = GetEditableSkeletonInternal()->GetSkeleton();
+				int32 BoneIndex = Skeleton.GetReferenceSkeleton().FindBoneIndex(BoneName);
+
+				float CurrentBlendScale = SelectedBlendProfile->GetBoneBlendScale(BoneIndex);
+
+				MenuBuilder.BeginSection("SkeletonTreeBlendProfileScales", LOCTEXT("BlendProfileContextOptions", "Blend Profile"));
+				{
+					FUIAction RecursiveSetScales;
+					RecursiveSetScales.ExecuteAction = FExecuteAction::CreateSP(this, &SSkeletonTree::RecursiveSetBlendProfileScales, CurrentBlendScale);
+					
+					MenuBuilder.AddMenuEntry
+						(
+						FText::Format(LOCTEXT("RecursiveSetBlendScales_Label", "Recursively Set Blend Scales To {0}"), FText::AsNumber(CurrentBlendScale)),
+						LOCTEXT("RecursiveSetBlendScales_ToolTip", "Sets all child bones to use the same blend profile scale as the selected bone"),
+						FSlateIcon(),
+						RecursiveSetScales
+						);
+				}
+				MenuBuilder.EndSection();
+			}
+
+			if(bShowingAdvancedOptions)
+			{
+				MenuBuilder.BeginSection("SkeletonTreeBoneTranslationRetargeting", LOCTEXT("BoneTranslationRetargetingHeader", "Bone Translation Retargeting"));
+				{
+					FUIAction RecursiveRetargetingSkeletonAction = FUIAction(FExecuteAction::CreateSP(this, &SSkeletonTree::SetBoneTranslationRetargetingModeRecursive, EBoneTranslationRetargetingMode::Skeleton));
+					FUIAction RecursiveRetargetingAnimationAction = FUIAction(FExecuteAction::CreateSP(this, &SSkeletonTree::SetBoneTranslationRetargetingModeRecursive, EBoneTranslationRetargetingMode::Animation));
+					FUIAction RecursiveRetargetingAnimationScaledAction = FUIAction(FExecuteAction::CreateSP(this, &SSkeletonTree::SetBoneTranslationRetargetingModeRecursive, EBoneTranslationRetargetingMode::AnimationScaled));
+					FUIAction RecursiveRetargetingAnimationRelativeAction = FUIAction(FExecuteAction::CreateSP(this, &SSkeletonTree::SetBoneTranslationRetargetingModeRecursive, EBoneTranslationRetargetingMode::AnimationRelative));
+
+					MenuBuilder.AddMenuEntry
+						(LOCTEXT("SetTranslationRetargetingSkeletonChildrenAction", "Recursively Set Translation Retargeting Skeleton")
+						, LOCTEXT("BoneTranslationRetargetingSkeletonToolTip", "Use translation from Skeleton.")
+						, FSlateIcon()
+						, RecursiveRetargetingSkeletonAction
+						);
+
+					MenuBuilder.AddMenuEntry
+						(LOCTEXT("SetTranslationRetargetingAnimationChildrenAction", "Recursively Set Translation Retargeting Animation")
+						, LOCTEXT("BoneTranslationRetargetingAnimationToolTip", "Use translation from animation.")
+						, FSlateIcon()
+						, RecursiveRetargetingAnimationAction
+						);
+
+					MenuBuilder.AddMenuEntry
+						(LOCTEXT("SetTranslationRetargetingAnimationScaledChildrenAction", "Recursively Set Translation Retargeting AnimationScaled")
+						, LOCTEXT("BoneTranslationRetargetingAnimationScaledToolTip", "Use translation from animation, scale length by Skeleton's proportions.")
+						, FSlateIcon()
+						, RecursiveRetargetingAnimationScaledAction
+						);
+
+					MenuBuilder.AddMenuEntry
+						(LOCTEXT("SetTranslationRetargetingAnimationRelativeChildrenAction", "Recursively Set Translation Retargeting AnimationRelative")
+						, LOCTEXT("BoneTranslationRetargetingAnimationRelativeToolTip", "Use relative translation from animation similar to an additive animation.")
+						, FSlateIcon()
+						, RecursiveRetargetingAnimationRelativeAction
+						);
+				}
+				MenuBuilder.EndSection();
+			}
+
+			MenuBuilder.BeginSection("SkeletonTreeBoneReductionForLOD", LOCTEXT("BoneReductionHeader", "LOD Bone Reduction"));
+			{
+				MenuBuilder.AddSubMenu(
+					LOCTEXT("SkeletonTreeBoneReductionForLOD_RemoveSelectedFromLOD", "Remove Selected..."),
+					FText::GetEmpty(),
+					FNewMenuDelegate::CreateStatic(&SSkeletonTree::CreateMenuForBoneReduction, this, LastCachedLODForPreviewMeshComponent, true)
+					);
+
+				MenuBuilder.AddSubMenu(
+					LOCTEXT("SkeletonTreeBoneReductionForLOD_RemoveChildrenFromLOD", "Remove Children..."),
+					FText::GetEmpty(),
+					FNewMenuDelegate::CreateStatic(&SSkeletonTree::CreateMenuForBoneReduction, this, LastCachedLODForPreviewMeshComponent, false)
+					);
+			}
+			MenuBuilder.EndSection();
+		}
+
+		if(BoneTreeSelection.HasSelectedOfType(ESkeletonTreeRowType::Socket))
+		{
+			MenuBuilder.BeginSection("SkeletonTreeSocketsActions", LOCTEXT( "SocketActions", "Selected Socket Actions" ) );
+
+			MenuBuilder.AddMenuEntry( Actions.CopySockets );
+
+			if(BoneTreeSelection.IsSingleOfTypeSelected(ESkeletonTreeRowType::Socket))
+			{
+				MenuBuilder.AddMenuEntry( FGenericCommands::Get().Rename, NAME_None, LOCTEXT("RenameSocket_Label", "Rename Socket"), LOCTEXT("RenameSocket_Tooltip", "Rename this socket") );
+
+				FDisplayedSocketInfo* DisplayedSocketInfo = static_cast< FDisplayedSocketInfo* >( BoneTreeSelection.GetSingleSelectedItem().Get() );
+
+				if ( DisplayedSocketInfo->IsSocketCustomized() && DisplayedSocketInfo->GetParentType() == ESocketParentType::Mesh )
+				{
+					MenuBuilder.AddMenuEntry( Actions.RemoveMeshSocket );
+				}
+
+				USkeletalMeshSocket* SelectedSocket = static_cast< USkeletalMeshSocket* >( DisplayedSocketInfo->GetData() );
+
+				// If the socket is on the skeleton, we have a valid mesh
+				// and there isn't one of the same name on the mesh, we can customize it
+				if ( DisplayedSocketInfo->CanCustomizeSocket() )
+				{
+					if ( DisplayedSocketInfo->GetParentType() == ESocketParentType::Skeleton )
+					{
+						MenuBuilder.AddMenuEntry( Actions.CreateMeshSocket );
+					}
+					else if ( DisplayedSocketInfo->GetParentType() == ESocketParentType::Mesh )
+					{
+						// If a socket is on the mesh only, then offer to promote it to the skeleton
+						MenuBuilder.AddMenuEntry( Actions.PromoteSocketToSkeleton );
+					}
+				}
+			}
+
+			MenuBuilder.EndSection();
+		}
+
+		MenuBuilder.BeginSection("SkeletonTreeAttachedAssets", LOCTEXT( "AttachedAssetsActionsHeader", "Attached Assets Actions" ) );
+
+		if ( BoneTreeSelection.IsSingleItemSelected() )
+		{
+			MenuBuilder.AddSubMenu(	LOCTEXT( "AttachNewAsset", "Add Preview Asset" ),
+				LOCTEXT ( "AttachNewAsset_ToolTip", "Attaches an asset to this part of the skeleton. Assets can also be dragged onto the skeleton from a content browser to attach" ),
+				FNewMenuDelegate::CreateSP( this, &SSkeletonTree::FillAttachAssetSubmenu, BoneTreeSelection.GetSingleSelectedItem() ) );
+		}
+
+		FUIAction RemoveAllAttachedAssets = FUIAction(	FExecuteAction::CreateSP( this, &SSkeletonTree::OnRemoveAllAssets ),
+			FCanExecuteAction::CreateSP( this, &SSkeletonTree::CanRemoveAllAssets ));
+
+		MenuBuilder.AddMenuEntry( LOCTEXT( "RemoveAllAttachedAssets", "Remove All Attached Assets" ),
+			LOCTEXT ( "RemoveAllAttachedAssets_ToolTip", "Removes all the attached assets from the skeleton and mesh." ),
+			FSlateIcon(), RemoveAllAttachedAssets );
+
+		MenuBuilder.EndSection();
+	}
+
+	return MenuBuilder.MakeWidget();
+}
+
+
+void SSkeletonTree::CreateMenuForBoneReduction(FMenuBuilder& MenuBuilder, SSkeletonTree * Widget, int32 LODIndex, bool bIncludeSelected)
+{
+	MenuBuilder.AddMenuEntry
+		(FText::FromString(FString::Printf(TEXT("From LOD %d and below"), LODIndex))
+		, FText::FromString(FString::Printf(TEXT("Remove Selected %s from current LOD %d and all lower LODs"), (bIncludeSelected) ? TEXT("bones") : TEXT("children"), LODIndex))
+		, FSlateIcon()
+		, FUIAction(FExecuteAction::CreateSP(Widget, &SSkeletonTree::RemoveFromLOD, LODIndex, bIncludeSelected, true))
+		);
+
+	MenuBuilder.AddMenuEntry
+		(FText::FromString(FString::Printf(TEXT("From LOD %d only"), LODIndex))
+		, FText::FromString(FString::Printf(TEXT("Remove selected %s from current LOD %d only"), (bIncludeSelected) ? TEXT("bones") : TEXT("children"), LODIndex))
+		, FSlateIcon()
+		, FUIAction(FExecuteAction::CreateSP(Widget, &SSkeletonTree::RemoveFromLOD, LODIndex, bIncludeSelected, false))
+		);
+}
+
+
+void SSkeletonTree::SetBoneTranslationRetargetingModeRecursive(EBoneTranslationRetargetingMode::Type NewRetargetingMode)
+{
+	TArray<FName> BoneNames;
+	FBoneTreeSelection TreeSelection(SkeletonTreeView->GetSelectedItems());
+	for ( auto ItemIt = TreeSelection.SelectedBones.CreateConstIterator(); ItemIt; ++ItemIt )
+	{
+		BoneNames.Add(*static_cast<FName*>( (*ItemIt)->GetData() ));
+	}
+
+	GetEditableSkeletonInternal()->SetBoneTranslationRetargetingModeRecursive(BoneNames, NewRetargetingMode);
+}
+
+void SSkeletonTree::RemoveFromLOD(int32 LODIndex, bool bIncludeSelected, bool bIncludeBelowLODs)
+{
+	// we cant do this without a preview scene
+	if (!GetPreviewScene().IsValid())
+	{
+		return;
+	}
+
+	UDebugSkelMeshComponent* PreviewMeshComponent = GetPreviewScene()->GetPreviewMeshComponent();
+	if (!PreviewMeshComponent->SkeletalMesh)
+	{
+		return;
+	}
+
+	// ask users you can't undo this change, and warn them
+	const FText Message(LOCTEXT("RemoveBonesFromLODWarning", "This action can't be undone. Would you like to continue?"));
+	if (FMessageDialog::Open(EAppMsgType::YesNo, Message) == EAppReturnType::Yes)
+	{
+		FBoneTreeSelection TreeSelection(SkeletonTreeView->GetSelectedItems());
+		const FReferenceSkeleton& RefSkeleton = GetEditableSkeletonInternal()->GetSkeleton().GetReferenceSkeleton();
+
+		TArray<FName> BonesToRemove;
+
+		for (auto ItemIt = TreeSelection.SelectedBones.CreateConstIterator(); ItemIt; ++ItemIt)
+		{
+			FName BoneName = *static_cast<FName*>((*ItemIt)->GetData());
+			int32 BoneIndex = RefSkeleton.FindBoneIndex(BoneName);
+			if (BoneIndex != INDEX_NONE)
+			{
+				if (bIncludeSelected)
+				{
+					PreviewMeshComponent->SkeletalMesh->AddBoneToReductionSetting(LODIndex, BoneName);
+					BonesToRemove.AddUnique(BoneName);
+				}
+				else
+				{
+					for (int32 ChildIndex = BoneIndex + 1; ChildIndex < RefSkeleton.GetNum(); ++ChildIndex)
+					{
+						if (RefSkeleton.GetParentIndex(ChildIndex) == BoneIndex)
+						{
+							FName ChildBoneName = RefSkeleton.GetBoneName(ChildIndex);
+							PreviewMeshComponent->SkeletalMesh->AddBoneToReductionSetting(LODIndex, ChildBoneName);
+							BonesToRemove.AddUnique(ChildBoneName);
+						}
+					}
+				}
+			}
+		}
+
+		int32 TotalLOD = PreviewMeshComponent->SkeletalMesh->LODInfo.Num();
+		IMeshUtilities& MeshUtilities = FModuleManager::Get().LoadModuleChecked<IMeshUtilities>("MeshUtilities");
+
+		if (bIncludeBelowLODs)
+		{
+			for (int32 Index = LODIndex + 1; Index < TotalLOD; ++Index)
+			{
+				MeshUtilities.RemoveBonesFromMesh(PreviewMeshComponent->SkeletalMesh, Index, &BonesToRemove);
+				PreviewMeshComponent->SkeletalMesh->AddBoneToReductionSetting(Index, BonesToRemove);
+			}
+		}
+
+		// remove from current LOD
+		MeshUtilities.RemoveBonesFromMesh(PreviewMeshComponent->SkeletalMesh, LODIndex, &BonesToRemove);
+		// update UI to reflect the change
+		OnLODSwitched();
+	}
+}
+
+void SSkeletonTree::OnCopyBoneNames()
+{
+	FBoneTreeSelection TreeSelection(SkeletonTreeView->GetSelectedItems());
+
+	if( TreeSelection.SelectedBones.Num() > 0 )
+	{
+		FString BoneNames;
+		for( auto ItemIt = TreeSelection.SelectedBones.CreateConstIterator(); ItemIt; ++ItemIt )
+		{
+			FName* BoneName = static_cast<FName*>( (*ItemIt)->GetData() );
+
+			BoneNames += BoneName->ToString();
+			BoneNames += "\r\n";
+		}
+		FPlatformMisc::ClipboardCopy( *BoneNames );
+	}
+}
+
+void SSkeletonTree::OnResetBoneTransforms()
+{
+	if (GetPreviewScene().IsValid())
+	{
+		UDebugSkelMeshComponent* PreviewComponent = GetPreviewScene()->GetPreviewMeshComponent();
+		check(PreviewComponent);
+		UAnimPreviewInstance* PreviewInstance = PreviewComponent->PreviewInstance;
+		check(PreviewInstance);
+
+		FBoneTreeSelection TreeSelection(SkeletonTreeView->GetSelectedItems());
+
+		if (TreeSelection.SelectedBones.Num() > 0)
+		{
+			bool bModified = false;
+			GEditor->BeginTransaction(LOCTEXT("SkeletonTree_ResetBoneTransforms", "Reset Bone Transforms"));
+
+			for (auto ItemIt = TreeSelection.SelectedBones.CreateConstIterator(); ItemIt; ++ItemIt)
+			{
+				const FName* BoneName = static_cast<FName*>((*ItemIt)->GetData());
+				check(BoneName);
+				const FAnimNode_ModifyBone* ModifiedBone = PreviewInstance->FindModifiedBone(*BoneName);
+				if (ModifiedBone != nullptr)
+				{
+					if (!bModified)
+					{
+						PreviewInstance->SetFlags(RF_Transactional);
+						PreviewInstance->Modify();
+						bModified = true;
+					}
+
+					PreviewInstance->RemoveBoneModification(*BoneName);
+				}
+			}
+
+			GEditor->EndTransaction();
+		}
+	}
+}
+
+void SSkeletonTree::OnCopySockets() const
+{
+	FBoneTreeSelection TreeSelection(SkeletonTreeView->GetSelectedItems());
+
+	int32 NumSocketsToCopy = TreeSelection.SelectedSockets.Num();
+	if ( NumSocketsToCopy > 0 )
+	{
+		FString SocketsDataString;
+
+		for ( auto ItemIt = TreeSelection.SelectedSockets.CreateConstIterator(); ItemIt; ++ItemIt )
+		{
+			FDisplayedSocketInfo* DisplayedSocketInfo = (*ItemIt).Get();
+			USkeletalMeshSocket* Socket = static_cast<USkeletalMeshSocket*>( DisplayedSocketInfo->GetData() );
+
+			SocketsDataString += SerializeSocketToString( Socket, DisplayedSocketInfo );
+		}
+
+		FString CopyString = FString::Printf( TEXT("%s\nNumSockets=%d\n%s"), *FEditableSkeleton::SocketCopyPasteHeader, NumSocketsToCopy, *SocketsDataString );
+
+		FPlatformMisc::ClipboardCopy( *CopyString );
+	}
+}
+
+FString SSkeletonTree::SerializeSocketToString( USkeletalMeshSocket* Socket, const FDisplayedSocketInfo* DisplayedSocketInfo ) const
+{
+	FString SocketString;
+
+	SocketString += FString::Printf( TEXT( "IsOnSkeleton=%s\n" ), DisplayedSocketInfo->GetParentType() == ESocketParentType::Skeleton ? TEXT( "1" ) : TEXT( "0" ) );
+
+	FStringOutputDevice Buffer;
+	const FExportObjectInnerContext Context;
+	UExporter::ExportToOutputDevice( &Context, Socket, nullptr, Buffer, TEXT( "copy" ), 0, PPF_Copy, false );
+	SocketString += Buffer;
+
+	return SocketString;
+}
+
+void SSkeletonTree::OnPasteSockets(bool bPasteToSelectedBone)
+{
+	FBoneTreeSelection TreeSelection(SkeletonTreeView->GetSelectedItems());
+
+	// Pasting sockets should only work if there is just one bone selected
+	if ( TreeSelection.IsSingleOfTypeSelected(ESkeletonTreeRowType::Bone) )
+	{
+		FName DestBoneName = bPasteToSelectedBone ? *static_cast<FName*>( TreeSelection.GetSingleSelectedItem()->GetData() ) : NAME_None;
+		GetEditableSkeletonInternal()->HandlePasteSockets(DestBoneName);
+
+		CreateFromSkeleton();
+	}
+}
+
+void SSkeletonTree::OnAddSocket()
+{
+	// This adds a socket to the currently selected bone in the SKELETON, not the MESH.
+	FBoneTreeSelection TreeSelection(SkeletonTreeView->GetSelectedItems());
+
+	// Can only add a socket to one bone
+	if (TreeSelection.IsSingleOfTypeSelected(ESkeletonTreeRowType::Bone))
+	{
+		FName BoneName = *static_cast<FName*>(TreeSelection.GetSingleSelectedItem()->GetData());
+		USkeletalMeshSocket* NewSocket = GetEditableSkeletonInternal()->HandleAddSocket(BoneName);
+		FSelectedSocketInfo SocketInfo(NewSocket, true);
+		SetSelectedSocket(SocketInfo);
+		CreateFromSkeleton(NewSocket);
+	}
+}
+
+void SSkeletonTree::OnCustomizeSocket()
+{
+	// This should only be called on a skeleton socket, it copies the 
+	// socket to the mesh so the user can edit it separately
+	FBoneTreeSelection TreeSelection(SkeletonTreeView->GetSelectedItems());
+
+	if(TreeSelection.IsSingleOfTypeSelected(ESkeletonTreeRowType::Socket))
+	{
+		USkeletalMeshSocket* SocketToCustomize = static_cast<USkeletalMeshSocket*>( TreeSelection.GetSingleSelectedItem()->GetData() );
+		GetEditableSkeletonInternal()->HandleCustomizeSocket(SocketToCustomize);
+		CreateFromSkeleton();
+	}
+}
+
+void SSkeletonTree::OnPromoteSocket()
+{
+	// This should only be called on a mesh socket, it copies the 
+	// socket to the skeleton so all meshes can use it
+	FBoneTreeSelection TreeSelection(SkeletonTreeView->GetSelectedItems());
+
+	// Can only customize one socket (CreateContextMenu() should prevent this firing!)
+	if(TreeSelection.IsSingleOfTypeSelected(ESkeletonTreeRowType::Socket))
+	{
+		USkeletalMeshSocket* SocketToPromote = static_cast<USkeletalMeshSocket*>( TreeSelection.GetSingleSelectedItem()->GetData() );
+		GetEditableSkeletonInternal()->HandlePromoteSocket(SocketToPromote);
+		CreateFromSkeleton();
+	}
+}
+
+void SSkeletonTree::FillAttachAssetSubmenu(FMenuBuilder& MenuBuilder, const FDisplayedTreeRowInfoPtr TargetItem)
+{
+	FContentBrowserModule& ContentBrowserModule = FModuleManager::Get().LoadModuleChecked<FContentBrowserModule>(TEXT("ContentBrowser"));
+
+	TArray<UClass*> FilterClasses = FComponentAssetBrokerage::GetSupportedAssets(USceneComponent::StaticClass());
+
+	//Clean up the selection so it is relevant to Persona
+	FilterClasses.RemoveSingleSwap(UBlueprint::StaticClass(), false); //Child actor components broker gives us blueprints which isn't wanted
+	FilterClasses.RemoveSingleSwap(USoundBase::StaticClass(), false); //No sounds wanted
+
+	FAssetPickerConfig AssetPickerConfig;
+	AssetPickerConfig.Filter.bRecursiveClasses = true;
+
+	for(int i = 0; i < FilterClasses.Num(); ++i)
+	{
+		AssetPickerConfig.Filter.ClassNames.Add(FilterClasses[i]->GetFName());
+	}
+
+
+	AssetPickerConfig.OnAssetSelected = FOnAssetSelected::CreateSP(this, &SSkeletonTree::OnAssetSelectedFromPicker, TargetItem);
+
+	TSharedRef<SWidget> MenuContent = SNew(SBox)
+	.WidthOverride(384)
+	.HeightOverride(500)
+	[
+		ContentBrowserModule.Get().CreateAssetPicker(AssetPickerConfig)
+	];
+	MenuBuilder.AddWidget( MenuContent, FText::GetEmpty(), true);
+}
+
+void SSkeletonTree::OnAssetSelectedFromPicker(const FAssetData& AssetData, const FDisplayedTreeRowInfoPtr TargetItem)
+{
+	FSlateApplication::Get().DismissAllMenus();
+	TArray<FAssetData> Assets;
+	Assets.Add(AssetData);
+
+	AttachAssetsToSkeletonTree(TargetItem, Assets);
+}
+
+void  SSkeletonTree::OnRemoveAllAssets()
+{
+	GetEditableSkeletonInternal()->HandleRemoveAllAssets(GetPreviewScene());
+
+	CreateFromSkeleton();
+}
+
+bool SSkeletonTree::CanRemoveAllAssets() const
+{
+	USkeletalMesh* const SkeletalMesh = GetEditableSkeletonInternal()->GetSkeletalMesh();
+
+	const bool bHasPreviewAttachedObjects = GetEditableSkeletonInternal()->GetSkeleton().PreviewAttachedAssetContainer.Num() > 0;
+	const bool bHasMeshPreviewAttachedObjects = ( SkeletalMesh && SkeletalMesh->PreviewAttachedAssetContainer.Num() );
+
+	return bHasPreviewAttachedObjects || bHasMeshPreviewAttachedObjects;
+}
+
+bool SSkeletonTree::CanRenameSelected() const
+{
+	FBoneTreeSelection TreeSelection(SkeletonTreeView->GetSelectedItems());
+	return TreeSelection.IsSingleOfTypeSelected(ESkeletonTreeRowType::Socket);
+}
+
+void SSkeletonTree::OnRenameSocket()
+{
+	FBoneTreeSelection TreeSelection(SkeletonTreeView->GetSelectedItems());
+
+	if(TreeSelection.IsSingleOfTypeSelected(ESkeletonTreeRowType::Socket))
+	{
+		SkeletonTreeView->RequestScrollIntoView(TreeSelection.GetSingleSelectedItem());
+		DeferredRenameRequest = TreeSelection.GetSingleSelectedItem();
+	}
+}
+
+void SSkeletonTree::OnSelectionChanged(TSharedPtr<FDisplayedTreeRowInfo> Selection, ESelectInfo::Type SelectInfo)
+{
+	if( Selection.IsValid() )
+	{
+		//Get all the selected items
+		FBoneTreeSelection TreeSelection(SkeletonTreeView->GetSelectedItems());
+
+		if (GetPreviewScene().IsValid())
+		{
+			UDebugSkelMeshComponent* PreviewComponent = GetPreviewScene()->GetPreviewMeshComponent();
+			if (TreeSelection.SelectedItems.Num() > 0 && PreviewComponent)
+			{
+				// pick the first settable bone from the selection
+				for (auto ItemIt = TreeSelection.SelectedItems.CreateConstIterator(); ItemIt; ++ItemIt)
+				{
+					FDisplayedTreeRowInfoPtr Item = *(ItemIt);
+
+					// Test SelectInfo so we don't end up in an infinite loop due to delegates calling each other
+					if (SelectInfo != ESelectInfo::Direct && Item->GetType() == ESkeletonTreeRowType::Bone)
+					{
+						FName BoneName = *static_cast<FName *>(Item->GetData());
+
+						// Get bone index
+						int32 BoneIndex = PreviewComponent->GetBoneIndex(BoneName);
+						if (BoneIndex != INDEX_NONE)
+						{
+							GetPreviewScene()->SetSelectedBone(BoneName);
+							OnObjectSelectedMulticast.Broadcast(nullptr);
+							break;
+						}
+					}
+					// Test SelectInfo so we don't end up in an infinite loop due to delegates calling each other
+					else if (SelectInfo != ESelectInfo::Direct && Item->GetType() == ESkeletonTreeRowType::Socket)
+					{
+						FDisplayedSocketInfo* DisplayedSocketInfo = static_cast<FDisplayedSocketInfo*>(Item.Get());
+						USkeletalMeshSocket* Socket = static_cast<USkeletalMeshSocket*>(Item->GetData());
+						FSelectedSocketInfo SocketInfo(Socket, DisplayedSocketInfo->GetParentType() == ESocketParentType::Skeleton);
+						GetPreviewScene()->SetSelectedSocket(SocketInfo);
+
+						OnObjectSelectedMulticast.Broadcast(SocketInfo.Socket);
+					}
+					else if (Item->GetType() == ESkeletonTreeRowType::AttachedAsset)
+					{
+						GetPreviewScene()->DeselectAll();
+
+						OnObjectSelectedMulticast.Broadcast(nullptr);
+					}
+				}
+				PreviewComponent->PostInitMeshObject(PreviewComponent->MeshObject);
+			}
+		}
+	}
+	else
+	{
+		if (GetPreviewScene().IsValid())
+		{
+			// Tell the preview scene if the user ctrl-clicked the selected bone/socket to de-select it
+			GetPreviewScene()->DeselectAll();
+		}
+
+		OnObjectSelectedMulticast.Broadcast(nullptr);
+	}
+}
+
+FReply SSkeletonTree::OnDropAssetToSkeletonTree(const FDisplayedTreeRowInfoPtr TargetItem, const FDragDropEvent& DragDropEvent)
+{
+	TSharedPtr<FAssetDragDropOp> DragDropOp = DragDropEvent.GetOperationAs<FAssetDragDropOp>();
+	if (DragDropOp.IsValid())
+	{
+		//Do we have some assets to attach?
+		if(DragDropOp->AssetData.Num() > 0)
+		{
+			AttachAssetsToSkeletonTree(TargetItem, DragDropOp->AssetData);
+		}
+		return FReply::Handled();
+	}
+	return FReply::Unhandled();
+}
+
+void SSkeletonTree::AttachAssetsToSkeletonTree(const FDisplayedTreeRowInfoPtr TargetItem, const TArray<FAssetData>& AssetData)
+{
+	bool bAllAssetWereLoaded = true;
+	TArray<UObject*> DroppedObjects;
+	for (int32 AssetIdx = 0; AssetIdx < AssetData.Num(); ++AssetIdx)
+	{
+		UObject* Object = AssetData[AssetIdx].GetAsset();
+		if ( Object != NULL )
+		{
+			DroppedObjects.Add( Object );
+		}
+		else
+		{
+			bAllAssetWereLoaded = false;
+		}
+	}
+
+	if(bAllAssetWereLoaded)
+	{
+		FName AttachToName = TargetItem->GetAttachName();
+		bool bAttachToMesh = TargetItem->GetType() == ESkeletonTreeRowType::Socket &&
+			static_cast<FDisplayedSocketInfo*>(TargetItem.Get())->GetParentType() == ESocketParentType::Mesh;
+		
+		GetEditableSkeletonInternal()->HandleAttachAssets(DroppedObjects, AttachToName, bAttachToMesh, GetPreviewScene());
+		CreateFromSkeleton();
+	}
+}
+
+void SSkeletonTree::OnItemScrolledIntoView( FDisplayedTreeRowInfoPtr InItem, const TSharedPtr<ITableRow>& InWidget)
+{
+	if(DeferredRenameRequest.IsValid())
+	{
+		DeferredRenameRequest->RequestRename();
+		DeferredRenameRequest.Reset();
+	}
+}
+
+void SSkeletonTree::OnTreeDoubleClick( FDisplayedTreeRowInfoPtr InItem )
+{
+	InItem->OnItemDoubleClicked();
+}
+
+void SSkeletonTree::SetTreeItemExpansionRecursive(TSharedPtr< FDisplayedTreeRowInfo > TreeItem, bool bInExpansionState) const
+{
+	SkeletonTreeView->SetItemExpansion(TreeItem, bInExpansionState);
+
+	// Recursively go through the children.
+	for (auto It = TreeItem->Children.CreateIterator(); It; ++It)
+	{
+		SetTreeItemExpansionRecursive(*It, bInExpansionState);
+	}
+}
+
+void SSkeletonTree::PostUndo()
+{
+	// Rebuild the tree view whenever we undo a change to the skeleton
+	CreateFromSkeleton();
+
+	CreateTreeColumns();
+
+	if (GetPreviewScene().IsValid())
+	{
+		GetPreviewScene()->DeselectAll();
+	}
+}
+
+void SSkeletonTree::OnFilterTextChanged( const FText& SearchText )
+{
+	FilterText = SearchText;
+
+	CreateFromSkeleton();
+}
+
+void SSkeletonTree::SetSkeletalMesh(USkeletalMesh* NewSkeletalMesh)
+{
+	if (GetPreviewScene().IsValid())
+	{
+		GetPreviewScene()->SetPreviewMesh(NewSkeletalMesh);
+	}
+
+	CreateFromSkeleton();
+}
+
+void SSkeletonTree::SetSelectedSocket( const FSelectedSocketInfo& SocketInfo )
+{
+	if (!bSelectingSocket)
+	{
+		TGuardValue<bool> RecursionGuard(bSelectingSocket, true);
+
+		// This function is called when something else selects a socket (i.e. *NOT* the user clicking on a row in the treeview)
+		// For example, this would be called if user clicked a socket hit point in the preview window
+
+		// Firstly, find which row (if any) contains the socket requested
+		for (auto SkeletonRowIt = DisplayMirror.CreateConstIterator(); SkeletonRowIt; ++SkeletonRowIt)
+		{
+			FDisplayedTreeRowInfoPtr SkeletonRow = *(SkeletonRowIt);
+
+			if (SkeletonRow->GetType() == ESkeletonTreeRowType::Socket && SkeletonRow->GetData() == SocketInfo.Socket)
+			{
+				SkeletonTreeView->SetSelection(SkeletonRow);
+				SkeletonTreeView->RequestScrollIntoView(SkeletonRow);
+			}
+		}
+
+		if (GetPreviewScene().IsValid())
+		{
+			GetPreviewScene()->SetSelectedSocket(SocketInfo);
+		}
+
+		OnObjectSelectedMulticast.Broadcast(SocketInfo.Socket);
+	}
+}
+
+void SSkeletonTree::SetSelectedBone( const FName& BoneName )
+{
+	if (!bSelectingBone)
+	{
+		TGuardValue<bool> RecursionGuard(bSelectingBone, true);
+		// This function is called when something else selects a bone (i.e. *NOT* the user clicking on a row in the treeview)
+		// For example, this would be called if user clicked a bone hit point in the preview window
+
+		// Find which row (if any) contains the bone requested
+		for (auto SkeletonRowIt = DisplayMirror.CreateConstIterator(); SkeletonRowIt; ++SkeletonRowIt)
+		{
+			FDisplayedTreeRowInfoPtr SkeletonRow = *(SkeletonRowIt);
+
+			if (SkeletonRow->GetType() == ESkeletonTreeRowType::Bone &&
+				*static_cast<FName*>(SkeletonRow->GetData()) == BoneName)
+			{
+				SkeletonTreeView->SetSelection(SkeletonRow);
+				SkeletonTreeView->RequestScrollIntoView(SkeletonRow);
+			}
+		}
+
+		if (GetPreviewScene().IsValid())
+		{
+			GetPreviewScene()->SetSelectedBone(BoneName);
+		}
+
+		OnObjectSelectedMulticast.Broadcast(nullptr);
+	}
+}
+
+void SSkeletonTree::DeselectAll()
+{
+	if (!bDeselectingAll)
+	{
+		TGuardValue<bool> RecursionGuard(bDeselectingAll, true);
+		SkeletonTreeView->ClearSelection();
+
+		if (GetPreviewScene().IsValid())
+		{
+			GetPreviewScene()->DeselectAll();
+		}
+
+		OnObjectSelectedMulticast.Broadcast(nullptr);
+	}
+}
+
+void SSkeletonTree::NotifyUser( FNotificationInfo& NotificationInfo )
+{
+	TSharedPtr<SNotificationItem> Notification = FSlateNotificationManager::Get().AddNotification( NotificationInfo );
+	if ( Notification.IsValid() )
+	{
+		Notification->SetCompletionState( SNotificationItem::CS_Fail );
+	}
+}
+
+TSharedRef< SWidget > SSkeletonTree::CreateFilterMenu()
+{
+	const FSkeletonTreeCommands& Actions = FSkeletonTreeCommands::Get();
+
+	const bool CloseAfterSelection = true;
+	FMenuBuilder MenuBuilder( CloseAfterSelection, UICommandList );
+
+	MenuBuilder.BeginSection("Bones", LOCTEXT("OptionsMenuHeading", "Options"));
+	{
+		MenuBuilder.AddMenuEntry(Actions.ShowRetargeting);
+	}
+	MenuBuilder.EndSection();
+
+	MenuBuilder.BeginSection("Bones", LOCTEXT( "BonesMenuHeading", "Bones" ) );
+	{
+		MenuBuilder.AddMenuEntry( Actions.ShowAllBones );
+		MenuBuilder.AddMenuEntry( Actions.ShowMeshBones );
+		MenuBuilder.AddMenuEntry(Actions.ShowLODBones);
+		MenuBuilder.AddMenuEntry( Actions.ShowWeightedBones );
+		MenuBuilder.AddMenuEntry( Actions.HideBones );
+	}
+	MenuBuilder.EndSection();
+
+	MenuBuilder.BeginSection("Sockets", LOCTEXT("SocketsMenuHeading", "Sockets"));
+	{
+		MenuBuilder.AddMenuEntry(Actions.ShowActiveSockets);
+		MenuBuilder.AddMenuEntry(Actions.ShowMeshSockets);
+		MenuBuilder.AddMenuEntry(Actions.ShowSkeletonSockets);
+		MenuBuilder.AddMenuEntry(Actions.ShowAllSockets);
+		MenuBuilder.AddMenuEntry(Actions.HideSockets);
+	}
+	MenuBuilder.EndSection();
+
+	return MenuBuilder.MakeWidget();
+}
+
+void SSkeletonTree::SetBoneFilter( EBoneFilter InBoneFilter )
+{
+	check( InBoneFilter < EBoneFilter::Count );
+	BoneFilter = InBoneFilter;
+
+	CreateFromSkeleton();
+}
+
+bool SSkeletonTree::IsBoneFilter( EBoneFilter InBoneFilter ) const
+{
+	return BoneFilter == InBoneFilter;
+}
+
+void SSkeletonTree::SetSocketFilter( ESocketFilter InSocketFilter )
+{
+	check( InSocketFilter < ESocketFilter::Count );
+	SocketFilter = InSocketFilter;
+
+	SetPreviewComponentSocketFilter();
+
+	CreateFromSkeleton();
+}
+
+void SSkeletonTree::SetPreviewComponentSocketFilter() const
+{
+	// Set the socket filter in the debug skeletal mesh component so the viewport can share the filter settings
+	if (GetPreviewScene().IsValid())
+	{
+		UDebugSkelMeshComponent* PreviewComponent = GetPreviewScene()->GetPreviewMeshComponent();
+
+		bool bAllOrActive = (SocketFilter == ESocketFilter::All || SocketFilter == ESocketFilter::Active);
+
+		if (PreviewComponent)
+		{
+			PreviewComponent->bMeshSocketsVisible = bAllOrActive || SocketFilter == ESocketFilter::Mesh;
+			PreviewComponent->bSkeletonSocketsVisible = bAllOrActive || SocketFilter == ESocketFilter::Skeleton;
+		}
+	}
+}
+
+bool SSkeletonTree::IsSocketFilter( ESocketFilter InSocketFilter ) const
+{
+	return SocketFilter == InSocketFilter;
+}
+
+// @todo combine these two functions to one if not used by anywhere else
+bool SSkeletonTree::IsBoneWeighted( int32 MeshBoneIndex, UDebugSkelMeshComponent* PreviewComponent ) const
+{
+	// MeshBoneIndex must be an index into the mesh's skeleton, *not* the source skeleton!!!
+
+	if ( !PreviewComponent || !PreviewComponent->SkeletalMesh || !PreviewComponent->SkeletalMesh->GetImportedResource()->LODModels.Num() )
+	{
+		// If there's no mesh, then this bone can't possibly be weighted!
+		return false;
+	}
+
+	//Get current LOD
+	const int32 LODIndex = FMath::Clamp(PreviewComponent->PredictedLODLevel, 0, PreviewComponent->SkeletalMesh->GetImportedResource()->LODModels.Num()-1);
+	FStaticLODModel& LODModel = PreviewComponent->SkeletalMesh->GetImportedResource()->LODModels[ LODIndex ];
+
+	//Check whether the bone is vertex weighted
+	int32 Index = LODModel.ActiveBoneIndices.Find( MeshBoneIndex );
+	
+	return Index != INDEX_NONE;
+}
+
+bool SSkeletonTree::IsBoneRequired(int32 MeshBoneIndex, UDebugSkelMeshComponent* PreviewComponent) const
+{
+	// MeshBoneIndex must be an index into the mesh's skeleton, *not* the source skeleton!!!
+
+	if (!PreviewComponent || !PreviewComponent->SkeletalMesh || !PreviewComponent->SkeletalMesh->GetImportedResource()->LODModels.Num())
+	{
+		// If there's no mesh, then this bone can't possibly be weighted!
+		return false;
+	}
+
+	//Get current LOD
+	const int32 LODIndex = FMath::Clamp(PreviewComponent->PredictedLODLevel, 0, PreviewComponent->SkeletalMesh->GetImportedResource()->LODModels.Num() - 1);
+	FStaticLODModel& LODModel = PreviewComponent->SkeletalMesh->GetImportedResource()->LODModels[LODIndex];
+
+	//Check whether the bone is vertex weighted
+	int32 Index = LODModel.RequiredBones.Find(MeshBoneIndex);
+
+	return Index != INDEX_NONE;
+}
+
+FText SSkeletonTree::GetFilterMenuTitle() const
+{
+	FText BoneFilterMenuText;
+
+	switch ( BoneFilter )
+	{
+	case EBoneFilter::All:
+		BoneFilterMenuText = LOCTEXT( "BoneFilterMenuAll", "All Bones" );
+		break;
+
+	case EBoneFilter::Mesh:
+		BoneFilterMenuText = LOCTEXT( "BoneFilterMenuMesh", "Mesh Bones" );
+		break;
+
+	case EBoneFilter::LOD:
+		BoneFilterMenuText = LOCTEXT("BoneFilterMenuLOD", "LOD Bones");
+		break;
+
+	case EBoneFilter::Weighted:
+		BoneFilterMenuText = LOCTEXT( "BoneFilterMenuWeighted", "Weighted Bones" );
+		break;
+
+	case EBoneFilter::None:
+		BoneFilterMenuText = LOCTEXT( "BoneFilterMenuHidden", "Bones Hidden" );
+		break;
+
+	default:
+		// Unknown mode
+		check( 0 );
+		break;
+	}
+
+	FText SocketFilterMenuText;
+
+	switch (SocketFilter)
+	{
+	case ESocketFilter::Active:
+		SocketFilterMenuText = LOCTEXT("SocketFilterMenuActive", "Active Sockets");
+		break;
+
+	case ESocketFilter::Mesh:
+		SocketFilterMenuText = LOCTEXT("SocketFilterMenuMesh", "Mesh Sockets");
+		break;
+
+	case ESocketFilter::Skeleton:
+		SocketFilterMenuText = LOCTEXT("SocketFilterMenuSkeleton", "Skeleton Sockets");
+		break;
+
+	case ESocketFilter::All:
+		SocketFilterMenuText = LOCTEXT("SocketFilterMenuAll", "All Sockets");
+		break;
+
+	case ESocketFilter::None:
+		SocketFilterMenuText = LOCTEXT("SocketFilterMenuHidden", "Sockets Hidden");
+		break;
+
+	default:
+		// Unknown mode
+		check(0);
+		break;
+	}
+
+	return FText::Format(LOCTEXT("FilterMenuLabelFormat", "{0}, {1}"), BoneFilterMenuText, SocketFilterMenuText);
+}
+
+bool SSkeletonTree::IsAddingSocketsAllowed() const
+{
+	if ( SocketFilter == ESocketFilter::Skeleton ||
+		SocketFilter == ESocketFilter::Active ||
+		SocketFilter == ESocketFilter::All )
+	{
+		return true;
+	}
+
+	return false;
+}
+
+FReply SSkeletonTree::OnKeyDown( const FGeometry& MyGeometry, const FKeyEvent& InKeyEvent )
+{
+	if ( UICommandList->ProcessCommandBindings( InKeyEvent ) )
+	{
+		return FReply::Handled();
+	}
+
+	return FReply::Unhandled();
+}
+
+void SSkeletonTree::OnDeleteSelectedRows()
+{
+	FBoneTreeSelection TreeSelection(SkeletonTreeView->GetSelectedItems());
+
+	if(TreeSelection.HasSelectedOfType(ESkeletonTreeRowType::AttachedAsset) || TreeSelection.HasSelectedOfType(ESkeletonTreeRowType::Socket))
+	{
+		FScopedTransaction Transaction( LOCTEXT( "SkeletonTreeDeleteSelected", "Delete selected sockets/meshes/bones from skeleton tree" ) );
+
+		DeleteAttachedAssets( TreeSelection.SelectedAssets );
+		DeleteSockets( TreeSelection.SelectedSockets );
+
+		CreateFromSkeleton();
+	}
+}
+
+void SSkeletonTree::DeleteAttachedAssets(const TArray<TSharedPtr<FDisplayedAttachedAssetInfo>>& InDisplayedAttachedAssetInfos)
+{
+	DeselectAll();
+
+	TArray<FPreviewAttachedObjectPair> AttachedObjects;
+	for(const TSharedPtr<FDisplayedAttachedAssetInfo>& AttachedAssetInfo : InDisplayedAttachedAssetInfos)
+	{
+		FPreviewAttachedObjectPair Pair;
+		Pair.SetAttachedObject(AttachedAssetInfo->GetAsset());
+		Pair.AttachedTo = AttachedAssetInfo->GetParentName();
+
+		AttachedObjects.Add(Pair);
+	}
+
+	GetEditableSkeletonInternal()->HandleDeleteAttachedAssets(AttachedObjects, GetPreviewScene());
+}
+
+void SSkeletonTree::DeleteSockets(const TArray<TSharedPtr<FDisplayedSocketInfo>>& InDisplayedSocketInfos)
+{
+	DeselectAll();
+
+	TArray<FSelectedSocketInfo> SocketInfo;
+
+	for (const TSharedPtr<FDisplayedSocketInfo>& DisplayedSocketInfo : InDisplayedSocketInfos)
+	{
+		USkeletalMeshSocket* SocketToDelete = static_cast<USkeletalMeshSocket*>( DisplayedSocketInfo->GetData() );
+		SocketInfo.Add(FSelectedSocketInfo(SocketToDelete, DisplayedSocketInfo->GetParentType() == ESocketParentType::Skeleton));
+	}
+
+	GetEditableSkeletonInternal()->HandleDeleteSockets(SocketInfo, GetPreviewScene());
+}
+
+void SSkeletonTree::AddAttachedAssets( const FPreviewAssetAttachContainer& AttachedObjects )
+{
+	for(auto Iter = AttachedObjects.CreateConstIterator(); Iter; ++Iter)
+	{
+		const FPreviewAttachedObjectPair& Pair = (*Iter);
+
+		if ( !FilterText.IsEmpty() && !Pair.GetAttachedObject()->GetName().Contains( FilterText.ToString() ) )
+		{
+			continue;
+		}
+
+		TSharedRef<FDisplayedAttachedAssetInfo> DisplayInfo = FDisplayedAttachedAssetInfo::Make(Pair.AttachedTo, Pair.GetAttachedObject(), SharedThis( this ) );
+		DisplayMirror.Add(DisplayInfo);
+
+		// for now it is a failure to not find where the asset is attached. Its possible that this might have to be changed to unloading the asset
+		// if there is a valid reason why the attach parent would not exist
+		if ( !AttachToParent( DisplayInfo, Pair.AttachedTo, (int32)( ESkeletonTreeRowType::Bone | ESkeletonTreeRowType::Socket ) ) )
+		{
+			// Just add it to the list if the parent bone isn't currently displayed
+			SkeletonRowList.Add( DisplayInfo );
+		}
+	}
+}
+
+void SSkeletonTree::OnChangeShowingAdvancedOptions()
+{
+	bShowingAdvancedOptions = !bShowingAdvancedOptions;
+	CreateTreeColumns();
+}
+
+bool SSkeletonTree::IsShowingAdvancedOptions() const
+{
+	return bShowingAdvancedOptions;
+}
+
+UBlendProfile* SSkeletonTree::GetSelectedBlendProfile()
+{
+	return BlendProfilePicker->GetSelectedBlendProfile();
+}
+
+FName SSkeletonTree::GetSelectedBlendProfileName() const
+{
+	return BlendProfilePicker->GetSelectedBlendProfileName();
+}
+
+void SSkeletonTree::OnBlendProfileSelected(UBlendProfile* NewProfile)
+{
+	CreateTreeColumns();
+}
+
+void SSkeletonTree::RecursiveSetBlendProfileScales(float InScaleToSet)
+{
+	UBlendProfile* SelectedBlendProfile = BlendProfilePicker->GetSelectedBlendProfile();
+	if(SelectedBlendProfile)
+	{
+		TArray<FName> BoneNames;
+		FBoneTreeSelection TreeSelection(SkeletonTreeView->GetSelectedItems());
+
+		for(auto ItemIt = TreeSelection.SelectedBones.CreateConstIterator(); ItemIt; ++ItemIt)
+		{
+			BoneNames.Add(*static_cast<FName*>((*ItemIt)->GetData()));
+		}
+
+		GetEditableSkeletonInternal()->RecursiveSetBlendProfileScales(SelectedBlendProfile->GetFName(), BoneNames, InScaleToSet);
+
+		CreateTreeColumns();
+	}
+}
+
+void SSkeletonTree::HandleTreeRefresh()
+{
+	CreateFromSkeleton();
+}
+
+void SSkeletonTree::PostRenameSocket(UObject* InAttachedObject, const FName& InOldName, const FName& InNewName)
+{
+	TSharedPtr<IPersonaPreviewScene> LinkedPreviewScene = GetPreviewScene();
+	if (LinkedPreviewScene.IsValid())
+	{
+		LinkedPreviewScene->RemoveAttachedObjectFromPreviewComponent(InAttachedObject, InOldName);
+		LinkedPreviewScene->AttachObjectToPreviewComponent(InAttachedObject, InNewName);
+	}
+}
+
+void SSkeletonTree::PostDuplicateSocket(UObject* InAttachedObject, const FName& InSocketName)
+{
+	TSharedPtr<IPersonaPreviewScene> LinkedPreviewScene = GetPreviewScene();
+	if (LinkedPreviewScene.IsValid())
+	{
+		LinkedPreviewScene->AttachObjectToPreviewComponent(InAttachedObject, InSocketName);
+	}
+
+	CreateFromSkeleton();
+}
+
+void RecursiveSetLODChange(UDebugSkelMeshComponent* PreviewComponent, TSharedPtr<FDisplayedTreeRowInfo> TreeRow)
+{
+	if (TreeRow.Get()->GetType() == ESkeletonTreeRowType::Bone)
+	{
+		static_cast<FDisplayedMeshBoneInfo*>(TreeRow.Get())->CacheLODChange(PreviewComponent);
+	}
+	
+	for (auto& Child : TreeRow->Children)
+	{
+		RecursiveSetLODChange(PreviewComponent, Child);
+	}
+}
+
+void SSkeletonTree::OnLODSwitched()
+{
+	if (GetPreviewScene().IsValid())
+	{
+		UDebugSkelMeshComponent* PreviewComponent = GetPreviewScene()->GetPreviewMeshComponent();
+
+		if (PreviewComponent)
+		{
+			LastCachedLODForPreviewMeshComponent = PreviewComponent->PredictedLODLevel;
+
+			if (BoneFilter == EBoneFilter::Weighted || BoneFilter == EBoneFilter::LOD)
+			{
+				CreateFromSkeleton();
+			}
+			else
+			{
+				for (auto ItemIt : SkeletonRowList)
+				{
+					RecursiveSetLODChange(PreviewComponent, ItemIt);
+				}
+			}
+		}
+	}
+}
+
+void SSkeletonTree::DuplicateAndSelectSocket(const FSelectedSocketInfo& SocketInfoToDuplicate, const FName& NewParentBoneName /*= FName()*/)
+{
+	USkeletalMeshSocket* NewSocket = GetEditableSkeleton()->DuplicateSocket(SocketInfoToDuplicate, NewParentBoneName);
+
+	if (GetPreviewScene().IsValid())
+	{
+		GetPreviewScene()->SetSelectedSocket(FSelectedSocketInfo(NewSocket, SocketInfoToDuplicate.bSocketIsOnSkeleton));
+	}
+
+	CreateFromSkeleton();
+
+	FSelectedSocketInfo NewSocketInfo(NewSocket, SocketInfoToDuplicate.bSocketIsOnSkeleton);
+	SetSelectedSocket(NewSocketInfo);
+}
+
+FSlateColor SSkeletonTree::GetFilterComboButtonForegroundColor() const
+{
+	static const FName InvertedForegroundName("InvertedForeground");
+	static const FName DefaultForegroundName("DefaultForeground");
+
+	if (FilterComboButton.IsValid())
+	{
+		return FilterComboButton->IsHovered() ? FEditorStyle::GetSlateColor(InvertedForegroundName) : FEditorStyle::GetSlateColor(DefaultForegroundName);
+	}
+	return FSlateColor::UseForeground();
+}
+
+#undef LOCTEXT_NAMESPACE
