Copying //UE4/Dev-Framework to //UE4/Dev-Main (Source: //UE4/Dev-Framework @ 3136612)

diff --git a/Engine/Source/Editor/SkeletonEditor/Private/SSkeletonTree.h b/Engine/Source/Editor/SkeletonEditor/Private/SSkeletonTree.h
new file mode 100644
--- /dev/null
+++ b/Engine/Source/Editor/SkeletonEditor/Private/SSkeletonTree.h
@@ -0,0 +1,656 @@
+// Copyright 1998-2016 Epic Games, Inc. All Rights Reserved.
+
+
+#pragma once
+
+#include "ISkeletonTree.h"
+#include "EditableSkeleton.h"
+#include "Engine/SkeletalMeshSocket.h"
+#include "Preferences/PersonaOptions.h"
+
+/** 
+ * Enum which determines what type a tree row is. Value is
+ * used as a flag for filtering tree items, so each goes up
+ * to the next bit value
+ */
+enum class ESkeletonTreeRowType : int32
+{
+	Bone = 1,
+	Socket = 2,
+	AttachedAsset = 4
+};
+
+ENUM_CLASS_FLAGS(ESkeletonTreeRowType);
+
+struct FDisplayedTreeRowInfo;
+class SBlendProfilePicker;
+
+typedef STreeView< TSharedPtr<FDisplayedTreeRowInfo> > SMeshSkeletonTreeRowType;
+typedef TSharedPtr< FDisplayedTreeRowInfo > FDisplayedTreeRowInfoPtr;
+
+//////////////////////////////////////////////////////////////////////////
+// FDisplayedTreeRowInfo
+
+struct FDisplayedTreeRowInfo : public TSharedFromThis<FDisplayedTreeRowInfo>
+{
+public:
+	TArray< TSharedPtr<FDisplayedTreeRowInfo> > Children;
+
+	// I don't like this, but our derived classes both hold completely different types of data
+	// so void* seems the only sensible thing without RTTI (which isn't sensible either!)
+	virtual void* GetData() = 0;
+	virtual ESkeletonTreeRowType GetType() const = 0;
+
+	/** Builds the table row widget to display this info */
+	virtual TSharedRef<ITableRow> MakeTreeRowWidget(
+		const TSharedRef<STableViewBase>& InOwnerTable,
+		FText InFilterText ) = 0;
+
+	/** Builds the slate widget for the name column */
+	virtual void GenerateWidgetForNameColumn( TSharedPtr< SHorizontalBox > Box, FText& FilterText, FIsSelected InIsSelected ) = 0;
+
+	/** Builds the slate widget for the data column */
+	virtual TSharedRef< SWidget > GenerateWidgetForDataColumn(const FName& DataColumnName) = 0;
+
+	/** Builds the slate widget for any inline data editing */
+	virtual TSharedRef< SWidget > GenerateInlineEditWidget(FText& FilterText, FIsSelected InIsSelected) { return SNullWidget::NullWidget; };
+
+	/** @return true if the item has an in-line editor widget */
+	virtual bool HasInlineEditor() const { return false; }
+
+	/** Toggle the expansion state of the inline editor */
+	virtual void ToggleInlineEditorExpansion() {}
+
+	/** Get the expansion state of the inline editor */
+	virtual bool IsInlineEditorExpanded() const { return false; }
+
+	/** Get the name of the item that this row represents */
+	virtual FName GetRowItemName() const PURE_VIRTUAL(FDisplayedTreeRowInfo::GetRowItemName, return FName(););
+
+	/** Return the name used to attach to this item */
+	virtual FName GetAttachName() const { return GetRowItemName(); }
+
+	/** Requests a rename on the the tree row item */
+	virtual void RequestRename() {};
+
+	/** Handler for when the user double clicks on this item in the tree */
+	virtual void OnItemDoubleClicked() {}
+
+	/** Get a shared reference to the editable skeleton that owns us */
+	TSharedRef<class FEditableSkeleton> GetEditableSkeleton() const;
+
+	/** Get a shared reference to the skeleton tree that owns us */
+	TSharedRef<class SSkeletonTree> GetSkeletonTree() const { return SkeletonTree.Pin().ToSharedRef(); }
+
+protected:
+	FDisplayedTreeRowInfo() {};
+
+	/** Skeleton tree that owns us */
+	TWeakPtr<class SSkeletonTree> SkeletonTree;
+};
+
+//////////////////////////////////////////////////////////////////////////
+// FDisplayedMeshBoneInfo
+
+struct FDisplayedMeshBoneInfo : public FDisplayedTreeRowInfo
+{
+public:
+	/** Static function for creating a new item, but ensures that you can only have a TSharedRef to one */
+	static TSharedRef<FDisplayedMeshBoneInfo> Make(
+		const FName& BoneName,
+		TWeakPtr<SSkeletonTree> InSkeletonTree)
+	{
+		FDisplayedMeshBoneInfo* DisplayedMeshBoneInfo = new FDisplayedMeshBoneInfo(BoneName);
+		DisplayedMeshBoneInfo->SkeletonTree = InSkeletonTree;
+
+		return MakeShareable( DisplayedMeshBoneInfo );
+	}
+
+	// Manual RTTI - not particularly elegant! :-(
+	virtual void* GetData() override { return &BoneName; }
+	virtual ESkeletonTreeRowType GetType() const override { return ESkeletonTreeRowType::Bone; }
+
+	/** Builds the table row widget to display this info */
+	virtual TSharedRef<ITableRow> MakeTreeRowWidget(
+		const TSharedRef<STableViewBase>& InOwnerTable,
+		FText InFilterText ) override;
+
+	/** Builds the slate widget for the name column */
+	virtual void GenerateWidgetForNameColumn( TSharedPtr< SHorizontalBox > Box, FText& FilterText, FIsSelected InIsSelected ) override;
+
+	/** Builds the slate widget for the data column */
+	virtual TSharedRef< SWidget > GenerateWidgetForDataColumn(const FName& DataColumnName) override;
+
+	/** Handle dragging a bone */
+	FReply OnDragDetected(const FGeometry& MyGeometry, const FPointerEvent& MouseEvent);
+
+	/** Return the name of the bone */
+	virtual FName GetRowItemName() const override {return BoneName;}
+
+	/** Set Translation Retargeting Mode for this bone. */
+	void SetBoneTranslationRetargetingMode(EBoneTranslationRetargetingMode::Type NewRetargetingMode);
+
+	/** Set current Blend Scale for this bone */
+	void SetBoneBlendProfileScale(float NewScale, bool bRecurse);
+
+	virtual ~FDisplayedMeshBoneInfo() {}
+
+	/** Cache to use Bold Font or not */
+	void CacheLODChange(class UDebugSkelMeshComponent* PreviewComponent);
+
+protected:
+	/** Hidden constructor, always use Make above */
+	FDisplayedMeshBoneInfo(const FName& InSource)
+		: BoneName(InSource)
+	{}
+
+	/** Hidden constructor, always use Make above */
+	FDisplayedMeshBoneInfo() {}
+
+private:
+	/** Gets the font for displaying bone text in the skeletal tree */
+	FSlateFontInfo GetBoneTextFont() const;
+
+	/** Get the text color based on bone part of skeleton or part of mesh */
+	FSlateColor GetBoneTextColor() const;
+
+	/** visibility of the icon */
+	EVisibility GetLODIconVisibility() const;
+
+	/** Function that returns the current tooltip for this bone, depending on how it's used by the mesh */
+	FText GetBoneToolTip();
+
+	/** Vary the color of the combo based on hover state */
+	FSlateColor GetRetargetingComboButtonForegroundColor() const;
+
+	/** The actual bone data that we create Slate widgets to display */
+	FName BoneName;
+
+	/** Weighted bone */
+	bool bWeightedBone;
+
+	/** Required bone */
+	bool bRequiredBone;
+
+	/** Reference to the retargetting combo button */
+	TSharedPtr<SComboButton> RetargetingComboButton;
+
+	/** Create menu for Bone Translation Retargeting Mode. */
+	TSharedRef< SWidget > CreateBoneTranslationRetargetingModeMenu();
+
+	/** Get Title for Bone Translation Retargeting Mode menu. */
+	FText GetTranslationRetargetingModeMenuTitle() const;
+
+	/** Callback from a slider widget if the text entry or slider movement is used */
+	void OnBlendSliderCommitted(float NewValue, ETextCommit::Type CommitType);
+};
+
+//////////////////////////////////////////////////////////////////////////
+// FDisplayedSocketInfo
+
+struct FDisplayedSocketInfo : public FDisplayedTreeRowInfo
+{
+public:
+	/** Static function for creating a new item, but ensures that you can only have a TSharedRef to one */
+	static TSharedRef<FDisplayedSocketInfo> Make(
+		USkeletalMeshSocket* Source,
+		ESocketParentType InParentType,
+		TWeakPtr<SSkeletonTree> InSkeletonTree,
+		bool bIsCustomized )
+	{
+		FDisplayedSocketInfo* DisplayedSocketInfo = new FDisplayedSocketInfo(Source, InParentType);
+		DisplayedSocketInfo->SkeletonTree = InSkeletonTree;
+		DisplayedSocketInfo->bIsCustomized = bIsCustomized;
+
+		return MakeShareable( DisplayedSocketInfo );
+	}
+
+	// Manual RTTI - not particularly elegant! :-(
+	virtual void* GetData() override { return SocketData; }
+	virtual ESkeletonTreeRowType GetType() const override { return ESkeletonTreeRowType::Socket; }
+
+	/** Builds the table row widget to display this info */
+	virtual TSharedRef<ITableRow> MakeTreeRowWidget(
+		const TSharedRef<STableViewBase>& InOwnerTable,
+		FText InFilterText ) override;
+
+	/** Builds the slate widget for the name column */
+	virtual void GenerateWidgetForNameColumn( TSharedPtr< SHorizontalBox > Box, FText& FilterText, FIsSelected InIsSelected ) override;
+
+	/** Builds the slate widget for any inline data editing */
+	virtual TSharedRef< SWidget > GenerateInlineEditWidget(FText& FilterText, FIsSelected InIsSelected) override;
+
+	/** @return true if the item has an in-line editor widget */
+	virtual bool HasInlineEditor() const override { return GetDefault<UPersonaOptions>()->bUseInlineSocketEditor; }
+
+	/** Toggle the expansion state of the inline editor */
+	virtual void ToggleInlineEditorExpansion() { bInlineEditorExpanded = !bInlineEditorExpanded; }
+
+	/** Get the expansion state of the inline editor */
+	virtual bool IsInlineEditorExpanded() const { return bInlineEditorExpanded; }
+
+	/** Builds the slate widget for the data column */
+	virtual TSharedRef< SWidget > GenerateWidgetForDataColumn(const FName& DataColumnName) override;
+
+	ESocketParentType GetParentType() const { return ParentType; }
+
+	/** Handle dragging a socket */
+	FReply OnDragDetected(const FGeometry& MyGeometry, const FPointerEvent& MouseEvent);
+
+	/** Return the name of the socket */
+	virtual FName GetRowItemName() const override {return SocketData->SocketName;}
+
+	/** Handle double clicking a socket */
+	virtual void OnItemDoubleClicked() override;
+
+	/** Is this socket customized */
+	bool IsSocketCustomized() const { return bIsCustomized; }
+
+	/** Can we customize this socket? */
+	bool CanCustomizeSocket() const;
+
+	/** Delegate for when the context menu requests a rename */
+	DECLARE_DELEGATE(FOnRenameRequested);
+	FOnRenameRequested OnRenameRequested;
+
+	/** Requests a rename on the socket item */
+	virtual void RequestRename() override;
+
+	/** Return socket name as FText for display in skeleton tree */
+	FText GetSocketNameAsText() const { return FText::FromName(SocketData->SocketName); }
+
+	virtual ~FDisplayedSocketInfo() {}
+
+protected:
+	/** Hidden constructor, always use Make above */
+	FDisplayedSocketInfo( USkeletalMeshSocket* InSource, ESocketParentType InParentType )
+		: SocketData( InSource )
+		, ParentType( InParentType )
+		, bInlineEditorExpanded( false )
+	{}
+
+	/** Hidden constructor, always use Make above */
+	FDisplayedSocketInfo() {}
+
+	/** Called when user is renaming a socket to verify the name **/
+	bool OnVerifySocketNameChanged( const FText& InText, FText& OutErrorMessage );
+
+	/** Called when user renames a socket **/
+	void OnCommitSocketName( const FText& InText, ETextCommit::Type CommitInfo );
+
+	/** Function that returns the current tooltip for this socket */
+	FText GetSocketToolTip();
+
+	/** Get the color to display for the socket name */
+	FSlateColor GetTextColor() const;
+
+	/** Pointer to the socket */
+	USkeletalMeshSocket*	SocketData;
+
+	/** This enum tells us whether the socket is on the skeleton or the mesh */
+	ESocketParentType ParentType;
+
+	/** Box for the user to type in the name - stored here so that SSkeletonTree can set the keyboard focus in Slate */
+	TSharedPtr<SEditableText> NameEntryBox;
+
+	/** True for sockets which exist on both the skeleton and mesh */
+	bool bIsCustomized;
+
+	/** Track expansion state of the in-line editor */
+	bool bInlineEditorExpanded;
+};
+
+//////////////////////////////////////////////////////////////////////////
+// FDisplayedAttachedAssetInfo
+
+struct FDisplayedAttachedAssetInfo : public FDisplayedTreeRowInfo
+{
+public:
+	/** Static function for creating a new item, but ensures that you can only have a TSharedRef to one */
+	static TSharedRef<FDisplayedAttachedAssetInfo> Make(const FName& InAttachedTo, UObject* InAsset, TWeakPtr<SSkeletonTree> InSkeletonTree);
+
+	// Manual RTTI - not particularly elegant! :-(
+	virtual void* GetData() override { return NULL; }
+	virtual ESkeletonTreeRowType GetType() const override { return ESkeletonTreeRowType::AttachedAsset; }
+
+	/** Builds the table row widget to display this info */
+	virtual TSharedRef<ITableRow> MakeTreeRowWidget(
+		const TSharedRef<STableViewBase>& InOwnerTable,
+		FText InFilterText ) override;
+
+	/** Builds the slate widget for the name column */
+	virtual void GenerateWidgetForNameColumn( TSharedPtr< SHorizontalBox > Box, FText& FilterText, FIsSelected InIsSelected ) override;
+
+	/** Builds the slate widget for the data column */
+	virtual TSharedRef< SWidget > GenerateWidgetForDataColumn(const FName& DataColumnName) override;
+
+	/** Return the name of the asset */
+	virtual FName GetRowItemName() const override {return FName( *Asset->GetName() ) ;}
+
+	/** Return the name used to attach to this item (in assets case return the item we are attached to */
+	virtual FName GetAttachName() const override { return GetParentName(); }
+
+	/** Returns the name of the socket/bone this asset is attached to */
+	const FName& GetParentName() const { return AttachedTo; }
+
+	/** Return the asset this info represents */
+	UObject* GetAsset() const { return Asset; }
+
+	/** Accessor for SCheckBox **/
+	ECheckBoxState IsAssetDisplayed() const;
+
+	/** Called when user toggles checkbox **/
+	void OnToggleAssetDisplayed( ECheckBoxState InCheckboxState );
+
+	/** Called when we need to get the state-based-image to show for the asset displayed checkbox */
+	const FSlateBrush* OnGetAssetDisplayedButtonImage() const;
+
+	/** Handler for when the user double clicks on this item in the tree */
+	virtual void OnItemDoubleClicked() override;
+
+	virtual ~FDisplayedAttachedAssetInfo() {}
+
+protected:
+	/** Hidden constructor, always use Make above */
+	FDisplayedAttachedAssetInfo( const FName& InAttachedTo, UObject* InAsset ) :
+		 AttachedTo(InAttachedTo), Asset(InAsset)
+	{}
+
+	/** The name of the socket/bone this asset is attached to */
+	const FName AttachedTo;
+
+	/** The attached asset that this tree item represents */
+	UObject* Asset;
+
+	/** The component of the attached asset */
+	TWeakObjectPtr<USceneComponent> AssetComponent;
+};
+//////////////////////////////////////////////////////////////////////////
+// SSkeletonTree
+
+class SSkeletonTree : public ISkeletonTree
+{
+public:
+	SLATE_BEGIN_ARGS( SSkeletonTree )
+		: _IsEditable(true)
+		{}
+
+	SLATE_ATTRIBUTE( bool, IsEditable )
+
+	SLATE_END_ARGS()
+public:
+	SSkeletonTree()
+		: bSelectingSocket(false)
+		, bSelectingBone(false)
+		
+	{}
+
+	~SSkeletonTree();
+
+	void Construct(const FArguments& InArgs, const TSharedRef<FEditableSkeleton>& InEditableSkeleton, const FSkeletonTreeArgs& InSkeletonTreeArgs);
+
+	/** ISkeletonTree interface */
+	virtual TSharedRef<class IEditableSkeleton> GetEditableSkeleton() const override { return EditableSkeleton.Pin().ToSharedRef(); }
+	virtual TSharedPtr<class IPersonaPreviewScene> GetPreviewScene() const override { return PreviewScene; }
+	virtual void SetSkeletalMesh(USkeletalMesh* NewSkeletalMesh) override;
+	virtual void SetSelectedSocket(const struct FSelectedSocketInfo& InSocketInfo) override;
+	virtual void SetSelectedBone(const FName& InBoneName) override;
+	virtual void DeselectAll() override;
+	virtual void DuplicateAndSelectSocket(const FSelectedSocketInfo& SocketInfoToDuplicate, const FName& NewParentBoneName = FName()) override;
+	virtual void RegisterOnObjectSelected(const FOnObjectSelected& Delegate) override
+	{
+		OnObjectSelectedMulticast.Add(Delegate);
+	}
+
+	virtual void UnregisterOnObjectSelected(SWidget* Widget) override
+	{
+		OnObjectSelectedMulticast.RemoveAll(Widget);
+	}
+
+	/** Creates the tree control and then populates */
+	void CreateTreeColumns();
+
+	/** Function to build the skeleton tree widgets from the source skeleton tree */
+	void CreateFromSkeleton(USkeletalMeshSocket* SocketToRename = nullptr);
+
+	/** This triggers a rebuild of the tree after undo to make the UI consistent with the real data */
+	void PostUndo();
+
+	/** Utility function to print notifications to the user */
+	void NotifyUser( FNotificationInfo& NotificationInfo );
+
+	/** Handle dropping something onto a skeleton bone tree item */
+	FReply OnDropAssetToSkeletonTree(const FDisplayedTreeRowInfoPtr TargetItem, const FDragDropEvent& DragDropEvent);
+
+	/** Attached the supplied assets to the tree to the specified attach item (bone/socket) */
+	void AttachAssetsToSkeletonTree(const FDisplayedTreeRowInfoPtr TargetItem, const TArray<FAssetData>& AssetData);
+
+	/** Returns true if a bone has vertices weighted against it */
+	bool IsBoneWeighted( int32 MeshBoneIndex, class UDebugSkelMeshComponent* PreviewComponent ) const;
+	bool IsBoneRequired(int32 MeshBoneIndex, class UDebugSkelMeshComponent* PreviewComponent) const;
+
+	/** Callback when an item is scrolled into view, handling calls to rename items */
+	void OnItemScrolledIntoView( FDisplayedTreeRowInfoPtr InItem, const TSharedPtr<ITableRow>& InWidget);
+
+	/** Callback for when the user double-clicks on an item in the tree */
+	void OnTreeDoubleClick( FDisplayedTreeRowInfoPtr InItem );
+
+	/** Handle recursive expansion/contraction of the tree */
+	void SetTreeItemExpansionRecursive(TSharedPtr< FDisplayedTreeRowInfo > TreeItem, bool bInExpansionState) const;
+
+	/** Set Bone Translation Retargeting Mode for bone selection, and their children. */
+	void SetBoneTranslationRetargetingModeRecursive(EBoneTranslationRetargetingMode::Type NewRetargetingMode);
+
+	/** Remove the selected bones from LOD of LODIndex when using simplygon **/
+	void RemoveFromLOD(int32 LODIndex, bool bIncludeSelected, bool bIncludeBelowLODs);
+
+	/** Called when the preview mesh is changed - simply rebuilds the skeleton tree for the new mesh */
+	void OnLODSwitched();
+
+	/** Gets the currently selected blend profile */
+	UBlendProfile* GetSelectedBlendProfile();
+
+	/** Get the name of the currently selected blend profile */
+	FName GetSelectedBlendProfileName() const;
+
+	/** Delegate handler for when the tree needs refreshing */
+	void HandleTreeRefresh();
+
+	/** Get a shared reference to the skeleton tree that owns us */
+	TSharedRef<FEditableSkeleton> GetEditableSkeletonInternal() const { return EditableSkeleton.Pin().ToSharedRef(); }
+
+	/** Update preview scene and tree after a socket rename */
+	void PostRenameSocket(UObject* InAttachedObject, const FName& InOldName, const FName& InNewName);
+
+	/** Update preview scene and tree after a socket duplication */
+	void PostDuplicateSocket(UObject* InAttachedObject, const FName& InSocketName);
+
+private:
+	/** Binds the commands in FSkeletonTreeCommands to functions in this class */
+	void BindCommands();
+
+	/** Create a widget for an entry in the tree from an info */
+	TSharedRef<ITableRow> MakeTreeRowWidget(TSharedPtr<FDisplayedTreeRowInfo> InInfo, const TSharedRef<STableViewBase>& OwnerTable);
+
+	/** Get all children for a given entry in the list */
+	void GetChildrenForInfo(TSharedPtr<FDisplayedTreeRowInfo> InInfo, TArray< TSharedPtr<FDisplayedTreeRowInfo> >& OutChildren);
+
+	/** Called to display context menu when right clicking on the widget */
+	TSharedPtr< SWidget > CreateContextMenu();
+
+	/** Called to display the filter menu */
+	TSharedRef< SWidget > CreateFilterMenu();
+
+	/** Function to copy selected bone name to the clipboard */
+	void OnCopyBoneNames();
+
+	/** Function to reset the transforms of selected bones */
+	void OnResetBoneTransforms();
+
+	/** Function to copy selected sockets to the clipboard */
+	void OnCopySockets() const;
+
+	/** Function to serialize a single socket to a string */
+	FString SerializeSocketToString( USkeletalMeshSocket* Socket, const FDisplayedSocketInfo* DisplayedSocketInfo ) const;
+
+	/** Function to paste selected sockets from the clipboard */
+	void OnPasteSockets(bool bPasteToSelectedBone);
+
+	/** Function to add a socket to the selected bone (skeleton, not mesh) */
+	void OnAddSocket();
+
+	/** Function to check if it is possible to rename the selected item */
+	bool CanRenameSelected() const;
+
+	/** Function to start renaming a socket */
+	void OnRenameSocket();
+
+	/** Function to customize a socket - this essentially copies a socket from the skeleton to the mesh */
+	void OnCustomizeSocket();
+
+	/** Function to promote a socket - this essentially copies a socket from the mesh to the skeleton */
+	void OnPromoteSocket();
+
+	/** Create content picker sub menu to allow users to pick an asset to attach */
+	void FillAttachAssetSubmenu(FMenuBuilder& MenuBuilder, const FDisplayedTreeRowInfoPtr TargetItem);
+
+	/** Helper function for asset picker that handles users choice */
+	void OnAssetSelectedFromPicker(const FAssetData& AssetData, const FDisplayedTreeRowInfoPtr TargetItem);
+
+	/** Context menu function to remove all attached assets */
+	void OnRemoveAllAssets();
+
+	/** Context menu function to control enabled/disabled status of remove all assets menu item */
+	bool CanRemoveAllAssets() const;
+
+	/** Functions to copy sockets from the skeleton to the mesh */
+	void OnCopySocketToMesh() {};
+
+	/** Callback function to be called when selection changes in the tree view widget. */
+	void OnSelectionChanged(TSharedPtr<FDisplayedTreeRowInfo> Selection, ESelectInfo::Type SelectInfo);
+
+	/** Filters the SListView when the user changes the search text box (NameFilterBox)	*/
+	void OnFilterTextChanged( const FText& SearchText );
+
+	/** Attach the given item to its parent */
+	bool AttachToParent( TSharedRef<FDisplayedTreeRowInfo> ItemToAttach, FName ParentName, int32 ItemsToInclude);
+
+	/** Add sockets from a TArray - separate function to avoid duplicating for skeleton and mesh */
+	void AddSocketsFromData(
+		const TArray< USkeletalMeshSocket* >& SocketArray,
+		ESocketParentType ParentType,
+		USkeletalMeshSocket* SocketToRename );
+
+	/** Sets which types of bone to show */
+	void SetBoneFilter( EBoneFilter InBoneFilter );
+
+	/** Queries the bone filter */
+	bool IsBoneFilter( EBoneFilter InBoneFilter ) const;
+
+	/** Sets which types of socket to show */
+	void SetSocketFilter( ESocketFilter InSocketFilter );
+
+	/** Queries the bone filter */
+	bool IsSocketFilter( ESocketFilter InSocketFilter ) const;
+
+	/** Returns the current text for the filter button - "All", "Mesh" or "Weighted" etc. */
+	FText GetFilterMenuTitle() const;
+
+	/** We can only add sockets in Active, Skeleton or All mode (otherwise they just disappear) */
+	bool IsAddingSocketsAllowed() const;
+
+	/** Handler for "Show Retargeting Options" check box IsChecked functionality */
+	bool IsShowingAdvancedOptions() const;
+
+	/**  Handler for when we change the "Show Retargeting Options" check box */
+	void OnChangeShowingAdvancedOptions();
+
+	/** This replicates the socket filter to the previewcomponent so that the viewport can use the same settings */
+	void SetPreviewComponentSocketFilter() const;
+
+	/** Override OnKeyDown */
+	virtual FReply OnKeyDown( const FGeometry& MyGeometry, const FKeyEvent& InKeyEvent );
+
+	/** Function to delete all the selected sockets/assets */
+	void OnDeleteSelectedRows();
+
+	/** Function to remove attached assets from the skeleton/mesh */
+	void DeleteAttachedAssets(const TArray<TSharedPtr<FDisplayedAttachedAssetInfo>>& InDisplayedAttachedAssetInfos );
+
+	/** Function to remove sockets from the skeleton/mesh */
+	void DeleteSockets(const TArray<TSharedPtr<FDisplayedSocketInfo>>& InDisplayedSocketInfos );
+
+	/** Add attached assets from a given TArray of them */
+	void AddAttachedAssets( const FPreviewAssetAttachContainer& AttachedObjects );
+
+	/** Called when the user selects a blend profile to edit */
+	void OnBlendProfileSelected(UBlendProfile* NewProfile);
+
+	/** Sets the blend scale for the selected bones and all of their children */
+	void RecursiveSetBlendProfileScales(float InScaleToSet);
+
+	/** Submenu creator handler for the given skeleton */
+	static void CreateMenuForBoneReduction(FMenuBuilder& MenuBuilder, SSkeletonTree* SkeletonTree, int32 LODIndex, bool bIncludeSelected);
+
+	/** Vary the foreground color of the filter button based on hover state */
+	FSlateColor GetFilterComboButtonForegroundColor() const;
+
+private:
+	/** Pointer back to the skeleton tree that owns us */
+	TWeakPtr<FEditableSkeleton> EditableSkeleton;
+
+	/** Link to a preview scene */
+	TSharedPtr<IPersonaPreviewScene> PreviewScene;
+
+	/** SSearchBox to filter the tree */
+	TSharedPtr<SSearchBox>	NameFilterBox;
+
+	/** The blend profile picker displaying the selected profile */
+	TSharedPtr<SBlendProfilePicker> BlendProfilePicker;
+
+	/** Widget user to hold the skeleton tree */
+	TSharedPtr<SOverlay> TreeHolder;
+
+	/** Widget used to display the skeleton hierarchy */
+	TSharedPtr<SMeshSkeletonTreeRowType> SkeletonTreeView;
+
+	/** A tree of bone info. Used by the BoneTreeView. */
+	TArray< TSharedPtr<struct FDisplayedTreeRowInfo> > SkeletonRowList;
+
+	/** A "mirror" of the tree as a flat array for easier searching */
+	TArray< TSharedRef<struct FDisplayedTreeRowInfo> > DisplayMirror;
+
+	/** Is this view editable */
+	TAttribute<bool> IsEditable;
+
+	/** Current text typed into NameFilterBox */
+	FText FilterText;
+
+	/** Commands that are bound to delegates*/
+	TSharedPtr<FUICommandList> UICommandList;
+
+	/** Current type of bones to show */
+	EBoneFilter BoneFilter;
+
+	/** Current type of sockets to show */
+	ESocketFilter SocketFilter;
+
+	bool bShowingAdvancedOptions;
+
+	/** Points to an item that is being requested to be renamed */
+	TSharedPtr<FDisplayedTreeRowInfo> DeferredRenameRequest;
+
+	/** Last Cached Preview Mesh Component LOD */
+	int32 LastCachedLODForPreviewMeshComponent;
+
+	/** Delegate for when a socket is selected by clicking its hit point */
+	FOnObjectSelectedMulticast OnObjectSelectedMulticast;
+
+	/** Selection recursion guard flags */
+	bool bSelectingSocket;
+	bool bSelectingBone;
+	bool bDeselectingAll;
+
+	/** Hold onto the filter combo button to set its foreground color */
+	TSharedPtr<SComboButton> FilterComboButton;
+}; 
diff --git a/Engine/Source/Editor/SkeletonEditor/Private/SkeletonEditorMode.cpp b/Engine/Source/Editor/SkeletonEditor/Private/SkeletonEditorMode.cpp
new file mode 100644
--- /dev/null
+++ b/Engine/Source/Editor/SkeletonEditor/Private/SkeletonEditorMode.cpp
@@ -0,0 +1,94 @@
+// Copyright 1998-2016 Epic Games, Inc. All Rights Reserved.
+
+#include "SkeletonEditorPrivatePCH.h"
+#include "SkeletonEditorMode.h"
+#include "SkeletonEditor.h"
+#include "ISkeletonEditorModule.h"
+#include "IPersonaToolkit.h"
+
+FSkeletonEditorMode::FSkeletonEditorMode(TSharedRef<FWorkflowCentricApplication> InHostingApp, TSharedRef<ISkeletonTree> InSkeletonTree)
+	: FApplicationMode(SkeletonEditorModes::SkeletonEditorMode)
+{
+	HostingAppPtr = InHostingApp;
+
+	TSharedRef<FSkeletonEditor> SkeletonEditor = StaticCastSharedRef<FSkeletonEditor>(InHostingApp);
+
+	ISkeletonEditorModule& SkeletonEditorModule = FModuleManager::LoadModuleChecked<ISkeletonEditorModule>("SkeletonEditor");
+	TabFactories.RegisterFactory(SkeletonEditorModule.CreateSkeletonTreeTabFactory(InHostingApp, InSkeletonTree));
+
+	FPersonaModule& PersonaModule = FModuleManager::LoadModuleChecked<FPersonaModule>("Persona");
+	TabFactories.RegisterFactory(PersonaModule.CreateDetailsTabFactory(InHostingApp, FOnDetailsCreated::CreateSP(&SkeletonEditor.Get(), &FSkeletonEditor::HandleDetailsCreated)));
+	TabFactories.RegisterFactory(PersonaModule.CreatePersonaViewportTabFactory(InHostingApp, InSkeletonTree, SkeletonEditor->GetPersonaToolkit()->GetPreviewScene(), SkeletonEditor->OnPostUndo, nullptr, FOnViewportCreated(), true, true));
+	TabFactories.RegisterFactory(PersonaModule.CreateAnimNotifiesTabFactory(InHostingApp, InSkeletonTree->GetEditableSkeleton(), SkeletonEditor->OnChangeAnimNotifies, SkeletonEditor->OnPostUndo, FOnObjectsSelected::CreateSP(&SkeletonEditor.Get(), &FSkeletonEditor::HandleObjectsSelected)));
+	TabFactories.RegisterFactory(PersonaModule.CreateAdvancedPreviewSceneTabFactory(InHostingApp, SkeletonEditor->GetPersonaToolkit()->GetPreviewScene()));
+	TabFactories.RegisterFactory(PersonaModule.CreateRetargetManagerTabFactory(InHostingApp, InSkeletonTree->GetEditableSkeleton(), SkeletonEditor->GetPersonaToolkit()->GetPreviewScene(), SkeletonEditor->OnPostUndo));
+	TabFactories.RegisterFactory(PersonaModule.CreateCurveViewerTabFactory(InHostingApp, InSkeletonTree->GetEditableSkeleton(), SkeletonEditor->GetPersonaToolkit()->GetPreviewScene(), SkeletonEditor->OnCurvesChanged, SkeletonEditor->OnPostUndo));
+	TabFactories.RegisterFactory(PersonaModule.CreateSkeletonSlotNamesTabFactory(InHostingApp, InSkeletonTree->GetEditableSkeleton(), SkeletonEditor->OnPostUndo, FOnObjectSelected::CreateSP(&SkeletonEditor.Get(), &FSkeletonEditor::HandleObjectSelected)));
+
+	TabLayout = FTabManager::NewLayout("Standalone_SkeletonEditor_Layout_v1.1")
+		->AddArea
+		(
+			FTabManager::NewPrimaryArea()
+			->SetOrientation(Orient_Vertical)
+			->Split
+			(
+				FTabManager::NewStack()
+				->SetSizeCoefficient(0.1f)
+				->SetHideTabWell(true)
+				->AddTab(InHostingApp->GetToolbarTabId(), ETabState::OpenedTab)
+			)
+			->Split
+			(
+				FTabManager::NewSplitter()
+				->SetSizeCoefficient(0.9f)
+				->SetOrientation(Orient_Horizontal)
+				->Split
+				(
+					FTabManager::NewStack()
+					->SetSizeCoefficient(0.2f)
+					->SetHideTabWell(true)
+					->AddTab(SkeletonEditorTabs::SkeletonTreeTab, ETabState::OpenedTab)
+					->AddTab(SkeletonEditorTabs::RetargetManagerTab, ETabState::ClosedTab)
+				)
+				->Split
+				(
+					FTabManager::NewStack()
+					->SetSizeCoefficient(0.6f)
+					->SetHideTabWell(true)
+					->AddTab(SkeletonEditorTabs::ViewportTab, ETabState::OpenedTab)
+				)
+				->Split
+				(
+					FTabManager::NewSplitter()
+					->SetSizeCoefficient(0.2f)
+					->SetOrientation(Orient_Vertical)
+					->Split
+					(
+						FTabManager::NewStack()
+						->SetSizeCoefficient(0.5f)
+						->SetHideTabWell(false)
+						->AddTab(SkeletonEditorTabs::DetailsTab, ETabState::OpenedTab)
+						->AddTab(SkeletonEditorTabs::AdvancedPreviewTab, ETabState::ClosedTab)
+					)
+					->Split
+					(
+						FTabManager::NewStack()
+						->SetSizeCoefficient(0.5f)
+						->SetHideTabWell(false)
+						->AddTab(SkeletonEditorTabs::AnimNotifiesTab, ETabState::OpenedTab)
+						->AddTab(SkeletonEditorTabs::CurveNamesTab, ETabState::OpenedTab)
+						->AddTab(SkeletonEditorTabs::SlotNamesTab, ETabState::ClosedTab)
+					)
+				)
+			)
+		);
+}
+
+void FSkeletonEditorMode::RegisterTabFactories(TSharedPtr<FTabManager> InTabManager)
+{
+	TSharedPtr<FWorkflowCentricApplication> HostingApp = HostingAppPtr.Pin();
+	HostingApp->RegisterTabSpawners(InTabManager.ToSharedRef());
+	HostingApp->PushTabFactories(TabFactories);
+
+	FApplicationMode::RegisterTabFactories(InTabManager);
+}
\ No newline at end of file
diff --git a/Engine/Source/Editor/SkeletonEditor/Private/SkeletonEditorMode.h b/Engine/Source/Editor/SkeletonEditor/Private/SkeletonEditorMode.h
new file mode 100644
--- /dev/null
+++ b/Engine/Source/Editor/SkeletonEditor/Private/SkeletonEditorMode.h
@@ -0,0 +1,23 @@
+// Copyright 1998-2016 Epic Games, Inc. All Rights Reserved.
+
+#pragma once
+
+#include "ApplicationMode.h"
+#include "WorkflowCentricApplication.h"
+#include "PersonaModule.h"
+
+class FSkeletonEditorMode : public FApplicationMode
+{
+public:
+	FSkeletonEditorMode(TSharedRef<class FWorkflowCentricApplication> InHostingApp, TSharedRef<class ISkeletonTree> InSkeletonTree);
+
+	/** FApplicationMode interface */
+	virtual void RegisterTabFactories(TSharedPtr<FTabManager> InTabManager) override;
+
+protected:
+	/** The hosting app */
+	TWeakPtr<class FWorkflowCentricApplication> HostingAppPtr;
+
+	/** The tab factories we support */
+	FWorkflowAllowedTabSet TabFactories;
+};
\ No newline at end of file
diff --git a/Engine/Source/Editor/SkeletonEditor/Private/SkeletonEditorPrivatePCH.h b/Engine/Source/Editor/SkeletonEditor/Private/SkeletonEditorPrivatePCH.h
new file mode 100644
--- /dev/null
+++ b/Engine/Source/Editor/SkeletonEditor/Private/SkeletonEditorPrivatePCH.h
@@ -0,0 +1,11 @@
+// Copyright 1998-2016 Epic Games, Inc. All Rights Reserved.
+
+#pragma once
+
+#include "UnrealEd.h"
+#include "SlateBasics.h"
+#include "Toolkits/AssetEditorToolkit.h"
+#include "Toolkits/IToolkit.h"
+#include "NotificationManager.h"
+#include "SNotificationList.h"
+#include "PropertyEditing.h"
diff --git a/Engine/Source/Editor/SkeletonEditor/Private/SkeletonTreeManager.cpp b/Engine/Source/Editor/SkeletonEditor/Private/SkeletonTreeManager.cpp
new file mode 100644
--- /dev/null
+++ b/Engine/Source/Editor/SkeletonEditor/Private/SkeletonTreeManager.cpp
@@ -0,0 +1,53 @@
+// Copyright 1998-2016 Epic Games, Inc. All Rights Reserved.
+
+#include "SkeletonEditorPrivatePCH.h"
+#include "SkeletonTreeManager.h"
+#include "EditableSkeleton.h"
+#include "ISkeletonTree.h"
+
+FSkeletonTreeManager& FSkeletonTreeManager::Get()
+{
+	static FSkeletonTreeManager TheManager;
+	return TheManager;
+}
+
+TSharedRef<ISkeletonTree> FSkeletonTreeManager::CreateSkeletonTree(class USkeleton* InSkeleton, class USkeletalMesh* InSkeletalMesh, const FSkeletonTreeArgs& InSkeletonTreeArgs)
+{
+	TSharedPtr<FEditableSkeleton> EditableSkeleton = CreateEditableSkeleton(InSkeleton, InSkeletalMesh);
+	TSharedPtr<ISkeletonTree> SkeletonTree = EditableSkeleton->CreateSkeletonTree(InSkeletonTreeArgs);
+
+	// compact skeletons that are no longer being edited
+	bool bRemoved = false;
+	do
+	{
+		bRemoved = false;
+		for (auto Iter = EditableSkeletons.CreateIterator(); Iter; ++Iter)
+		{
+			TSharedPtr<FEditableSkeleton> CompactionCandidate = Iter.Value().Pin();
+			if (!CompactionCandidate.IsValid() || !CompactionCandidate->IsEdited())
+			{
+				Iter.RemoveCurrent();
+				bRemoved = true;
+				break;
+			}
+		}
+	}
+	while (bRemoved);
+
+	return SkeletonTree.ToSharedRef();
+}
+
+TSharedRef<FEditableSkeleton> FSkeletonTreeManager::CreateEditableSkeleton(class USkeleton* InSkeleton, class USkeletalMesh* InSkeletalMesh /*= nullptr*/)
+{
+	TWeakPtr<FEditableSkeleton>* EditableSkeletonPtr = EditableSkeletons.Find(InSkeleton);
+	if (EditableSkeletonPtr == nullptr || !(*EditableSkeletonPtr).IsValid())
+	{
+		TSharedRef<FEditableSkeleton> NewEditableSkeleton = MakeShareable(new FEditableSkeleton(InSkeleton, InSkeletalMesh));
+		EditableSkeletons.Add(InSkeleton, NewEditableSkeleton);
+		return NewEditableSkeleton;
+	}
+	else
+	{
+		return (*EditableSkeletonPtr).Pin().ToSharedRef();
+	}
+}
diff --git a/Engine/Source/Editor/SkeletonEditor/Private/SkeletonTreeManager.h b/Engine/Source/Editor/SkeletonEditor/Private/SkeletonTreeManager.h
new file mode 100644
--- /dev/null
+++ b/Engine/Source/Editor/SkeletonEditor/Private/SkeletonTreeManager.h
@@ -0,0 +1,25 @@
+// Copyright 1998-2016 Epic Games, Inc. All Rights Reserved.
+
+#pragma once
+
+/** Central registry of skeleton trees */
+class FSkeletonTreeManager
+{
+public:
+	/** Singleton access */
+	static FSkeletonTreeManager& Get();
+
+	/** Create a skeleton tree for the requested skeleton */
+	TSharedRef<class ISkeletonTree> CreateSkeletonTree(class USkeleton* InSkeleton, class USkeletalMesh* InSkeletalMesh, const struct FSkeletonTreeArgs& InSkeletonTreeArgs);
+
+	/** Edit a USkeleton via FEditableSkeleton */
+	TSharedRef<class FEditableSkeleton> CreateEditableSkeleton(class USkeleton* InSkeleton, class USkeletalMesh* InSkeletalMesh = nullptr);
+
+private:
+	/** Hidden constructor */
+	FSkeletonTreeManager() {}
+
+private:
+	/** Map from skeletons to editable skeletons */
+	TMap<class USkeleton*, TWeakPtr<class FEditableSkeleton>> EditableSkeletons;
+};
diff --git a/Engine/Source/Editor/SkeletonEditor/Public/IEditableSkeleton.h b/Engine/Source/Editor/SkeletonEditor/Public/IEditableSkeleton.h
new file mode 100644
--- /dev/null
+++ b/Engine/Source/Editor/SkeletonEditor/Public/IEditableSkeleton.h
@@ -0,0 +1,157 @@
+// Copyright 1998-2016 Epic Games, Inc. All Rights Reserved.
+
+#pragma once
+
+#include "Animation/SmartName.h"
+
+/** Delegate fired when a set of smart names is removed */
+DECLARE_MULTICAST_DELEGATE_TwoParams(FOnSmartNameRemoved, const FName& /*InContainerName*/, const TArray<SmartName::UID_Type>& /*InNameUids*/);
+
+/** Enum which tells us whether the parent of a socket is the skeleton or skeletal mesh */
+enum class ESocketParentType : uint8
+{
+	Skeleton,
+	Mesh
+};
+
+/** Interface used to deal with skeletons that are currently being edited */
+class IEditableSkeleton
+{
+public:
+	/** Get the skeleton. Const - to modify the skeleton you should use the wrapper methods below */
+	virtual const class USkeleton& GetSkeleton() const = 0;
+
+	/** Get the blend profiles that this skeleton currently contains */
+	virtual const TArray<class UBlendProfile*>& GetBlendProfiles() const  = 0;
+
+	/** Get the named blend profile @return nullptr if none was found */
+	virtual class UBlendProfile* GetBlendProfile(const FName& InBlendProfileName) = 0;
+
+	/** Create the named blend profile */
+	virtual class UBlendProfile* CreateNewBlendProfile(const FName& InBlendProfileName) = 0;
+
+	/** Remove the specifed blend profile */
+	virtual void RemoveBlendProfile(UBlendProfile* InBlendProfile) = 0;
+
+	/** Set the blend profile scale for the specified bone */
+	virtual void SetBlendProfileScale(const FName& InBlendProfileName, const FName& InBoneName, float InNewScale, bool bInRecurse) = 0;
+
+	/** Handle drag-drop socket duplication */
+	virtual class USkeletalMeshSocket* DuplicateSocket(const struct FSelectedSocketInfo& SocketInfoToDuplicate, const FName& NewParentBoneName) = 0;
+
+	/** Rename a socket on the mesh or the skeleton */
+	virtual void RenameSocket(const FName& OldSocketName, const FName& NewSocketName) = 0;
+
+	/** Set the parent of a socket */
+	virtual void SetSocketParent(const FName& SocketName, const FName& NewParentName) = 0;
+
+	/** Function to tell you if a socket name already exists. SocketParentType determines where we will look to see if the socket exists, i.e. mesh or skeleton */
+	virtual bool DoesSocketAlreadyExist(const class USkeletalMeshSocket* InSocket, const FText& InSocketName, ESocketParentType SocketParentType) const = 0;
+
+	/** Add a new smart name. @return true if name was succesfully added */
+	virtual bool AddSmartname(const FName& InContainerName, const FName& InNewName, FSmartName& OutSmartName) = 0;
+
+	/** Rename the specified smart name */
+	virtual void RenameSmartname(const FName& InContainerName, SmartName::UID_Type InNameUid, const FName& InNewName) = 0;
+
+	/** Remove the specified smart name */
+	virtual void RemoveSmartname(const FName& InContainerName, SmartName::UID_Type InNameUid) = 0;
+
+	/** Remove all the specified smart names and fixup animations that use them */
+	virtual void RemoveSmartnamesAndFixupAnimations(const FName& InContainerName, const TArray<SmartName::UID_Type>& InNameUids) = 0;
+
+	/** Sets Material Meta Data for the curve */
+	virtual void SetCurveMetaDataMaterial(const FSmartName& CurveName, bool bOverrideMaterial) = 0;
+	/**
+	 * Makes sure all attached objects are valid and removes any that aren't.
+	 *
+	 * @return		NumberOfBrokenAssets
+	 */
+	virtual int32 ValidatePreviewAttachedObjects() = 0;
+
+	/** 
+	 * Merge has failed, then Recreate BoneTree
+	 * 
+	 * This will invalidate all animations that were linked before, but this is needed 
+	 * 
+	 * @param InSkelMesh			: Mesh to build from. 
+	 * 
+	 * @return true if success
+	 */	
+	virtual void RecreateBoneTree(class USkeletalMesh* NewPreviewMesh) = 0;
+
+	/** 
+	 * Delete anim notifies by name
+	 * @return the number of animations modified
+	 */	
+	virtual int32 DeleteAnimNotifies(const TArray<FName>& InotifyNames) = 0;
+
+	/** 
+	 * Rename a notify
+	 * @return the number of animations modified
+	 */	
+	virtual int32 RenameNotify(const FName& NewName, const FName& OldName) = 0;
+
+	/** Populates OutAssets with the AnimSequences that match this current skeleton */
+	virtual void GetCompatibleAnimSequences(TArray<class FAssetData>& OutAssets) = 0;
+
+	/** Set the preview mesh in the skeleton*/
+	virtual void SetPreviewMesh(class USkeletalMesh* InSkeletalMesh) = 0;
+
+	/** Load any additional meshes we may have */
+	virtual void LoadAdditionalPreviewSkeletalMeshes() = 0;
+
+	/** Set the additional skeletal meshes we use when previewing this skeleton */
+	virtual void SetAdditionalPreviewSkeletalMeshes(class UPreviewMeshCollection* InSkeletalMeshes) = 0;
+
+	/** Rename the specified retarget source */
+	virtual void RenameRetargetSource(const FName& InOldName, const FName& InNewName) = 0;
+
+	/** Add a retarget source */
+	virtual void AddRetargetSource(const FName& InName, USkeletalMesh* InReferenceMesh) = 0;
+
+	/** Delete retarget sources */
+	virtual void DeleteRetargetSources(const TArray<FName>& InRetargetSourceNames) = 0;
+
+	/** Refresh retarget sources */
+	virtual void RefreshRetargetSources(const TArray<FName>& InRetargetSourceNames) = 0;
+
+	/** Refreshes the rig config, validating the mappings */
+	virtual void RefreshRigConfig() = 0;
+
+	/** Set the rig config  */
+	virtual void SetRigConfig(class URig* InRig) = 0;
+
+	/** Set a rig bone mapping */
+	virtual void SetRigBoneMapping(const FName& InNodeName, const FName& InBoneName) = 0;
+
+	/** Set multiple bone mappings */
+	virtual void SetRigBoneMappings(const TMap<FName, FName>& InMappings) = 0;
+
+	/** Remove any bones that are not used by any skeletal meshes */
+	virtual void RemoveUnusedBones() = 0;
+
+	/** Create RefLocalPoses from InSkeletalMesh */
+	virtual void UpdateSkeletonReferencePose(class USkeletalMesh* InSkeletalMesh) = 0;
+
+	/** Add a slot group name */
+	virtual bool AddSlotGroupName(const FName& InSlotName) = 0;
+
+	/** Set the group name of the specified slot */
+	virtual void SetSlotGroupName(const FName& InSlotName, const FName& InGroupName) = 0;
+
+	/** Delete a slot name */
+	virtual void DeleteSlotName(const FName& InSlotName) = 0;
+
+	/** Delete a slot group */
+	virtual void DeleteSlotGroup(const FName& InGroupName) = 0;
+
+	/** Rename a slot name */
+	virtual void RenameSlotName(const FName& InOldSlotName, const FName& InNewSlotName) = 0;
+
+	/** Register a delegate to be called when a set of smart names are removed */
+	virtual FDelegateHandle RegisterOnSmartNameRemoved(const FOnSmartNameRemoved::FDelegate& InOnSmartNameRemoved) = 0;
+
+	/** Register a delegate to be called when a set of smart names are removed */
+	virtual void UnregisterOnSmartNameRemoved(FDelegateHandle InHandle) = 0;
+};
\ No newline at end of file
diff --git a/Engine/Source/Editor/SkeletonEditor/Public/ISkeletonEditor.h b/Engine/Source/Editor/SkeletonEditor/Public/ISkeletonEditor.h
new file mode 100644
--- /dev/null
+++ b/Engine/Source/Editor/SkeletonEditor/Public/ISkeletonEditor.h
@@ -0,0 +1,10 @@
+// Copyright 1998-2016 Epic Games, Inc. All Rights Reserved.
+
+#pragma once
+
+#include "WorkflowCentricApplication.h"
+#include "IHasPersonaToolkit.h"
+
+class ISkeletonEditor : public FWorkflowCentricApplication, public IHasPersonaToolkit
+{
+};
diff --git a/Engine/Source/Editor/SkeletonEditor/Public/ISkeletonTree.h b/Engine/Source/Editor/SkeletonEditor/Public/ISkeletonTree.h
new file mode 100644
--- /dev/null
+++ b/Engine/Source/Editor/SkeletonEditor/Public/ISkeletonTree.h
@@ -0,0 +1,80 @@
+// Copyright 1998-2016 Epic Games, Inc. All Rights Reserved.
+
+#pragma once
+
+// Called when a bone is selected
+DECLARE_MULTICAST_DELEGATE_OneParam(FOnObjectSelectedMulticast, UObject* /* InObject */);
+
+// Called when an object is selected
+typedef FOnObjectSelectedMulticast::FDelegate FOnObjectSelected;
+
+/** Init params for a skeleton tree widget */
+struct FSkeletonTreeArgs
+{
+	FSkeletonTreeArgs(FSimpleMulticastDelegate& InOnPostUndo)
+		: OnPostUndo(InOnPostUndo)
+	{}
+
+	/** Multicast delegate the tree will register with for undo messages */
+	FSimpleMulticastDelegate& OnPostUndo;
+
+	/** Delegate called by the tree when a socket is selected */
+	FOnObjectSelected OnObjectSelected;
+
+	/** Optional preview scene that we can pair with */
+	TSharedPtr<class IPersonaPreviewScene> PreviewScene;
+};
+
+/** Enum which tells us what type of bones we should be showing */
+enum class EBoneFilter : uint8
+{
+	All,
+	Mesh,
+	LOD,
+	Weighted, /** only showing weighted bones of current LOD */
+	None,
+	Count
+};
+
+/** Enum which tells us what type of sockets we should be showing */
+enum class ESocketFilter : uint8
+{
+	Active,
+	Mesh,
+	Skeleton,
+	All,
+	None,
+	Count
+};
+
+/** Interface used to deal with skeleton editing UI */
+class ISkeletonTree : public SCompoundWidget
+{
+public:
+	/** Get editable skeleton that this widget is editing */
+	virtual TSharedRef<class IEditableSkeleton> GetEditableSkeleton() const = 0;
+
+	/** Get preview scene that this widget is editing */
+	virtual TSharedPtr<class IPersonaPreviewScene> GetPreviewScene() const = 0;
+
+	/** Set the skeletal mesh we optionally work with */
+	virtual void SetSkeletalMesh(class USkeletalMesh* NewSkeletalMesh) = 0;
+
+	/** Set the selected socket */
+	virtual void SetSelectedSocket(const struct FSelectedSocketInfo& InSocketInfo) = 0;
+
+	/** Set the selected bone */
+	virtual void SetSelectedBone(const FName& InBoneName) = 0;
+
+	/** Deselect everything that is currently selected */
+	virtual void DeselectAll() = 0;
+
+	/** Duplicate the socket and select it */
+	virtual void DuplicateAndSelectSocket(const FSelectedSocketInfo& SocketInfoToDuplicate, const FName& NewParentBoneName = FName()) = 0;
+
+	/** Registers a delegate to be called when the selected object is changed */
+	virtual void RegisterOnObjectSelected(const FOnObjectSelected& Delegate) = 0;
+
+	/** Unregisters a delegate to be called when the selected object is changed */
+	virtual void UnregisterOnObjectSelected(SWidget* Widget) = 0;
+};
\ No newline at end of file
diff --git a/Engine/Source/Editor/StaticMeshEditor/Private/SStaticMeshEditorViewport.cpp b/Engine/Source/Editor/StaticMeshEditor/Private/SStaticMeshEditorViewport.cpp
--- a/Engine/Source/Editor/StaticMeshEditor/Private/SStaticMeshEditorViewport.cpp
+++ b/Engine/Source/Editor/StaticMeshEditor/Private/SStaticMeshEditorViewport.cpp
@@ -94,7 +94,7 @@ void SStaticMeshEditorViewport::Construct(const FArguments& InArgs)
 {
 	//PreviewScene = new FAdvancedPreviewScene(FPreviewScene::ConstructionValues(), 
 
-	PreviewScene.SetFloorOffset(-InArgs._ObjectToEdit->ExtendedBounds.Origin.Z + InArgs._ObjectToEdit->ExtendedBounds.BoxExtent.Z);
+	PreviewScene->SetFloorOffset(-InArgs._ObjectToEdit->ExtendedBounds.Origin.Z + InArgs._ObjectToEdit->ExtendedBounds.BoxExtent.Z);
 
 	StaticMeshEditorPtr = InArgs._StaticMeshEditor;
 
@@ -120,7 +120,7 @@ void SStaticMeshEditorViewport::Construct(const FArguments& InArgs)
 }
 
 SStaticMeshEditorViewport::SStaticMeshEditorViewport()
-	: PreviewScene(FPreviewScene::ConstructionValues())
+	: PreviewScene(MakeShareable(new FAdvancedPreviewScene(FPreviewScene::ConstructionValues())))
 {
 
 }
@@ -224,7 +224,7 @@ void SStaticMeshEditorViewport::UpdatePreviewSocketMeshes()
 			{
 				// Handle removing an old component
 				UStaticMeshComponent* SocketPreviewMeshComponent = SocketPreviewMeshComponents[i];
-				PreviewScene.RemoveComponent(SocketPreviewMeshComponent);
+				PreviewScene->RemoveComponent(SocketPreviewMeshComponent);
 				SocketPreviewMeshComponents.RemoveAt(i, SocketedComponentCount - i);
 				break;
 			}
@@ -236,7 +236,7 @@ void SStaticMeshEditorViewport::UpdatePreviewSocketMeshes()
 				if(i >= SocketedComponentCount)
 				{
 					SocketPreviewMeshComponent = NewObject<UStaticMeshComponent>();
-					PreviewScene.AddComponent(SocketPreviewMeshComponent, FTransform::Identity);
+					PreviewScene->AddComponent(SocketPreviewMeshComponent, FTransform::Identity);
 					SocketPreviewMeshComponents.Add(SocketPreviewMeshComponent);
 					SocketPreviewMeshComponent->SnapTo(PreviewMeshComponent, Socket->SocketName);
 				}
@@ -267,7 +267,7 @@ void SStaticMeshEditorViewport::SetPreviewMesh(UStaticMesh* InStaticMesh)
 	PreviewMeshComponent->StaticMesh = InStaticMesh;
 
 	FTransform Transform = FTransform::Identity;
-	PreviewScene.AddComponent( PreviewMeshComponent, Transform );
+	PreviewScene->AddComponent( PreviewMeshComponent, Transform );
 
 	EditorViewportClient->SetPreviewMesh(InStaticMesh, PreviewMeshComponent);
 }
@@ -281,7 +281,7 @@ void SStaticMeshEditorViewport::UpdatePreviewMesh(UStaticMesh* InStaticMesh)
 			UStaticMeshComponent* SocketPreviewMeshComponent = SocketPreviewMeshComponents[i];
 			if( SocketPreviewMeshComponent )
 			{
-				PreviewScene.RemoveComponent(SocketPreviewMeshComponent);
+				PreviewScene->RemoveComponent(SocketPreviewMeshComponent);
 			}
 		}
 		SocketPreviewMeshComponents.Empty();
@@ -289,7 +289,7 @@ void SStaticMeshEditorViewport::UpdatePreviewMesh(UStaticMesh* InStaticMesh)
 
 	if (PreviewMeshComponent)
 	{
-		PreviewScene.RemoveComponent(PreviewMeshComponent);
+		PreviewScene->RemoveComponent(PreviewMeshComponent);
 		PreviewMeshComponent = NULL;
 	}
 
@@ -300,7 +300,7 @@ void SStaticMeshEditorViewport::UpdatePreviewMesh(UStaticMesh* InStaticMesh)
 	// Update streaming data for debug viewmode feedback
 	PreviewMeshComponent->UpdateStreamingSectionData(FTexCoordScaleMap());
 
-	PreviewScene.AddComponent(PreviewMeshComponent,FTransform::Identity);
+	PreviewScene->AddComponent(PreviewMeshComponent,FTransform::Identity);
 
 	const int32 SocketCount = InStaticMesh->Sockets.Num();
 	SocketPreviewMeshComponents.Reserve(SocketCount);
@@ -315,7 +315,7 @@ void SStaticMeshEditorViewport::UpdatePreviewMesh(UStaticMesh* InStaticMesh)
 			SocketPreviewMeshComponent->SetStaticMesh(Socket->PreviewStaticMesh);
 			SocketPreviewMeshComponent->SnapTo(PreviewMeshComponent, Socket->SocketName);
 			SocketPreviewMeshComponents.Add(SocketPreviewMeshComponent);
-			PreviewScene.AddComponent(SocketPreviewMeshComponent, FTransform::Identity);
+			PreviewScene->AddComponent(SocketPreviewMeshComponent, FTransform::Identity);
 		}
 	}
 
@@ -404,7 +404,7 @@ FStaticMeshEditorViewportClient& SStaticMeshEditorViewport::GetViewportClient()
 
 TSharedRef<FEditorViewportClient> SStaticMeshEditorViewport::MakeEditorViewportClient()
 {
-	EditorViewportClient = MakeShareable( new FStaticMeshEditorViewportClient(StaticMeshEditorPtr, SharedThis(this), PreviewScene, StaticMesh, NULL) );
+	EditorViewportClient = MakeShareable( new FStaticMeshEditorViewportClient(StaticMeshEditorPtr, SharedThis(this), PreviewScene.ToSharedRef(), StaticMesh, NULL) );
 
 	EditorViewportClient->bSetListenerPosition = false;
 
diff --git a/Engine/Source/Editor/StaticMeshEditor/Private/SStaticMeshEditorViewport.h b/Engine/Source/Editor/StaticMeshEditor/Private/SStaticMeshEditorViewport.h
--- a/Engine/Source/Editor/StaticMeshEditor/Private/SStaticMeshEditorViewport.h
+++ b/Engine/Source/Editor/StaticMeshEditor/Private/SStaticMeshEditorViewport.h
@@ -90,7 +90,7 @@ public:
 	// End of ICommonEditorViewportToolbarInfoProvider interface
 
 	/** Returns the preview scene being renderd in the viewport */
-	FAdvancedPreviewScene& GetPreviewScene() { return PreviewScene; }
+	TSharedRef<FAdvancedPreviewScene> GetPreviewScene() { return PreviewScene.ToSharedRef(); }
 protected:
 	/** SEditorViewport interface */
 	virtual TSharedRef<FEditorViewportClient> MakeEditorViewportClient() override;
@@ -129,7 +129,7 @@ private:
 	TWeakPtr<IStaticMeshEditor> StaticMeshEditorPtr;
 
 	/** The scene for this viewport. */
-	FAdvancedPreviewScene PreviewScene;
+	TSharedPtr<FAdvancedPreviewScene> PreviewScene;
 
 	/** Editor viewport client */
 	TSharedPtr<class FStaticMeshEditorViewportClient> EditorViewportClient;
diff --git a/Engine/Source/Editor/StaticMeshEditor/Private/StaticMeshEditor.cpp b/Engine/Source/Editor/StaticMeshEditor/Private/StaticMeshEditor.cpp
--- a/Engine/Source/Editor/StaticMeshEditor/Private/StaticMeshEditor.cpp
+++ b/Engine/Source/Editor/StaticMeshEditor/Private/StaticMeshEditor.cpp
@@ -568,8 +568,7 @@ void FStaticMeshEditor::BuildSubTools()
 		LODLevelCombo->SetSelectedItem(LODLevels[0]);
 	}
 
-	AdvancedPreviewSettingsWidget = SNew(SAdvancedPreviewDetailsTab)
-		.PreviewScenePtr(&Viewport->GetPreviewScene());
+	AdvancedPreviewSettingsWidget = SNew(SAdvancedPreviewDetailsTab, Viewport->GetPreviewScene());
 }
 
 FName FStaticMeshEditor::GetToolkitFName() const
diff --git a/Engine/Source/Editor/StaticMeshEditor/Private/StaticMeshEditorViewportClient.cpp b/Engine/Source/Editor/StaticMeshEditor/Private/StaticMeshEditorViewportClient.cpp
--- a/Engine/Source/Editor/StaticMeshEditor/Private/StaticMeshEditorViewportClient.cpp
+++ b/Engine/Source/Editor/StaticMeshEditor/Private/StaticMeshEditorViewportClient.cpp
@@ -49,8 +49,8 @@ namespace {
 	static float AmbientCubemapIntensity = 0.4f;
 }
 
-FStaticMeshEditorViewportClient::FStaticMeshEditorViewportClient(TWeakPtr<IStaticMeshEditor> InStaticMeshEditor, const TSharedRef<SStaticMeshEditorViewport>& InStaticMeshEditorViewport, FAdvancedPreviewScene& InPreviewScene, UStaticMesh* InPreviewStaticMesh, UStaticMeshComponent* InPreviewStaticMeshComponent)
-	: FEditorViewportClient(nullptr, &InPreviewScene, StaticCastSharedRef<SEditorViewport>(InStaticMeshEditorViewport))
+FStaticMeshEditorViewportClient::FStaticMeshEditorViewportClient(TWeakPtr<IStaticMeshEditor> InStaticMeshEditor, const TSharedRef<SStaticMeshEditorViewport>& InStaticMeshEditorViewport, const TSharedRef<FAdvancedPreviewScene>& InPreviewScene, UStaticMesh* InPreviewStaticMesh, UStaticMeshComponent* InPreviewStaticMeshComponent)
+	: FEditorViewportClient(nullptr, &InPreviewScene.Get(), StaticCastSharedRef<SEditorViewport>(InStaticMeshEditorViewport))
 	, StaticMeshEditorPtr(InStaticMeshEditor)
 	, StaticMeshEditorViewportPtr(InStaticMeshEditorViewport)
 {
@@ -1226,8 +1226,8 @@ void FStaticMeshEditorViewportClient::SetAdvancedShowFlagsForScene()
 
 void FStaticMeshEditorViewportClient::SetFloorAndEnvironmentVisibility(const bool bVisible)
 {
-	AdvancedPreviewScene->SetFloorVisibility(bVisible);
-	AdvancedPreviewScene->SetEnvironmentVisibility(bVisible);
+	AdvancedPreviewScene->SetFloorVisibility(bVisible, true);
+	AdvancedPreviewScene->SetEnvironmentVisibility(bVisible, true);
 }
 
 void FStaticMeshEditorViewportClient::SetPreviewMesh(UStaticMesh* InStaticMesh, UStaticMeshComponent* InStaticMeshComponent)
diff --git a/Engine/Source/Editor/StaticMeshEditor/Private/StaticMeshEditorViewportClient.h b/Engine/Source/Editor/StaticMeshEditor/Private/StaticMeshEditorViewportClient.h
--- a/Engine/Source/Editor/StaticMeshEditor/Private/StaticMeshEditorViewportClient.h
+++ b/Engine/Source/Editor/StaticMeshEditor/Private/StaticMeshEditorViewportClient.h
@@ -11,7 +11,7 @@ class FAdvancedPreviewScene;
 class FStaticMeshEditorViewportClient : public FEditorViewportClient, public TSharedFromThis<FStaticMeshEditorViewportClient>
 {
 public:
-	FStaticMeshEditorViewportClient(TWeakPtr<IStaticMeshEditor> InStaticMeshEditor, const TSharedRef<SStaticMeshEditorViewport>& InStaticMeshEditorViewport, FAdvancedPreviewScene& InPreviewScene, UStaticMesh* InPreviewStaticMesh, UStaticMeshComponent* InPreviewStaticMeshComponent);
+	FStaticMeshEditorViewportClient(TWeakPtr<IStaticMeshEditor> InStaticMeshEditor, const TSharedRef<SStaticMeshEditorViewport>& InStaticMeshEditorViewport, const TSharedRef<FAdvancedPreviewScene>& InPreviewScene, UStaticMesh* InPreviewStaticMesh, UStaticMeshComponent* InPreviewStaticMeshComponent);
 	~FStaticMeshEditorViewportClient();
 
 	// FEditorViewportClient interface
diff --git a/Engine/Source/Editor/UnrealEd/Classes/Animation/DebugSkelMeshComponent.h b/Engine/Source/Editor/UnrealEd/Classes/Animation/DebugSkelMeshComponent.h
--- a/Engine/Source/Editor/UnrealEd/Classes/Animation/DebugSkelMeshComponent.h
+++ b/Engine/Source/Editor/UnrealEd/Classes/Animation/DebugSkelMeshComponent.h
@@ -12,7 +12,7 @@ struct FSelectedSocketInfo
 
 	/** Default constructor */
 	FSelectedSocketInfo()
-		: Socket( NULL )
+		: Socket(nullptr)
 		, bSocketIsOnSkeleton( false )
 	{
 
@@ -26,13 +26,21 @@ struct FSelectedSocketInfo
 
 	}
 
+	bool IsValid() const
+	{
+		return Socket != nullptr;
+	}
+
+	void Reset()
+	{
+		Socket = nullptr;
+	}
+
 	/** The socket we have selected */
 	class USkeletalMeshSocket* Socket;
 
 	/** true if on skeleton, false if on mesh */
 	bool bSocketIsOnSkeleton;
-
-
 };
 
 /** Different modes for Persona's Turn Table. */
@@ -86,10 +94,6 @@ class UDebugSkelMeshComponent : public USkeletalMeshComponent
 	UPROPERTY(transient)
 	uint32 bDrawBinormals:1;
 
-	/** CPU skinning rendering - only for previewing in Persona */
-	UPROPERTY(transient)
-	uint32 bCPUSkinning : 1;
-
 	/** Socket hit points viewing */
 	UPROPERTY(transient)
 	uint32 bDrawSockets:1;
@@ -143,9 +147,6 @@ class UDebugSkelMeshComponent : public USkeletalMeshComponent
 
 	/** Storage of Source Animation Pose for when bDisplaySourceAnimation == true, as they have to be calculated */
 	TArray<FTransform> SourceAnimationPoses;
-
-	/** Color render mode enum value - 0 - none, 1 - tangent, 2 - normal, 3 - mirror, 4 - bone weighting */
-	//var native transient int ColorRenderMode;
 	
 	/** Array of bones to render bone weights for */
 	UPROPERTY(transient)
@@ -154,12 +155,6 @@ class UDebugSkelMeshComponent : public USkeletalMeshComponent
 	/** Array of morphtargets to render verts for */
 	UPROPERTY(transient)
 	TArray<class UMorphTarget*> MorphTargetOfInterests;
-	
-	/** Array of sockets to render manipulation widgets for
-	/	Storing a pointer to the actual socket rather than a name, as we don't care here
-	/	whether the socket is on the skeleton or the mesh! */
-	UPROPERTY(transient)
-	TArray< FSelectedSocketInfo > SocketsOfInterest;
 
 	/** Array of materials to restore when not rendering blend weights */
 	UPROPERTY(transient)
@@ -175,9 +170,6 @@ class UDebugSkelMeshComponent : public USkeletalMeshComponent
 	UPROPERTY(transient)
 	bool bIsUsingInGameBounds;
 	
-	/** true if wind effects on clothing is enabled */
-	bool bEnableWind;
-
 	//~ Begin USceneComponent Interface.
 	virtual FBoxSphereBounds CalcBounds(const FTransform& LocalToWorld) const override;
 	//~ End USceneComponent Interface.
@@ -203,7 +195,7 @@ class UDebugSkelMeshComponent : public USkeletalMeshComponent
 
 	//~ Begin SkeletalMeshComponent Interface
 	virtual void InitAnim(bool bForceReinit) override;
-	virtual bool IsWindEnabled() const override;
+	virtual bool IsWindEnabled() const override { return true; }
 	//~ End SkeletalMeshComponent Interface
 	// Preview.
 	// @todo document
diff --git a/Engine/Source/Editor/UnrealEd/Classes/Preferences/PersonaOptions.h b/Engine/Source/Editor/UnrealEd/Classes/Preferences/PersonaOptions.h
--- a/Engine/Source/Editor/UnrealEd/Classes/Preferences/PersonaOptions.h
+++ b/Engine/Source/Editor/UnrealEd/Classes/Preferences/PersonaOptions.h
@@ -8,6 +8,7 @@
 //=============================================================================
 
 #pragma once
+#include "Engine/EngineBaseTypes.h"
 #include "PersonaOptions.generated.h"
 
 UCLASS(hidecategories=Object, config=EditorPerProjectUserSettings)
@@ -64,6 +65,15 @@ class UNREALED_API UPersonaOptions : public UObject
 	UPROPERTY(EditAnywhere, config, Category = Options)
 	FLinearColor BranchingPointTimingNodeColor;
 
+	UPROPERTY(EditAnywhere, config, Category = Options)
+	bool bUseStandaloneAnimationEditors;
+
+	UPROPERTY(EditAnywhere, config, Category = Options)
+	bool bUseInlineSocketEditor;
+
+	UPROPERTY(EditAnywhere, config, Category = Options)
+	bool bAllowPreviewMeshCollectionsToSelectFromDifferentSkeletons;
+
 public:
 	void SetViewportBackgroundColor( const FLinearColor& InViewportBackgroundColor);
 	void SetShowGrid( bool bInShowGrid );
diff --git a/Engine/Source/Editor/UnrealEd/Private/Animation/DebugSkelMeshComponent.cpp b/Engine/Source/Editor/UnrealEd/Private/Animation/DebugSkelMeshComponent.cpp
--- a/Engine/Source/Editor/UnrealEd/Private/Animation/DebugSkelMeshComponent.cpp
+++ b/Engine/Source/Editor/UnrealEd/Private/Animation/DebugSkelMeshComponent.cpp
@@ -81,9 +81,6 @@ UDebugSkelMeshComponent::UDebugSkelMeshComponent(const FObjectInitializer& Objec
 	bDisplayRawAnimation = false;
 	bDisplayNonRetargetedPose = false;
 
-	// wind is turned off in the editor by default
-	bEnableWind = false;
-
 	bMeshSocketsVisible = true;
 	bSkeletonSocketsVisible = true;
 
@@ -301,11 +298,13 @@ void UDebugSkelMeshComponent::InitAnim(bool bForceReinit)
 		AnimScriptInstance = PreviewInstance;
 		AnimScriptInstance->InitializeAnimation();
 	}
-}
 
-bool UDebugSkelMeshComponent::IsWindEnabled() const
-{
-	return bEnableWind;
+	if(PostProcessAnimInstance)
+	{
+		// Add the same settings as the preview instance in this case.
+		PostProcessAnimInstance->RootMotionMode = ERootMotionMode::RootMotionFromEverything;
+		PostProcessAnimInstance->bCanUseParallelUpdateAnimation = false;
+	}
 }
 
 void UDebugSkelMeshComponent::EnablePreview(bool bEnable, UAnimationAsset* PreviewAsset)
@@ -447,96 +446,100 @@ void UDebugSkelMeshComponent::RefreshBoneTransforms(FActorComponentTickFunction*
 	// Run regular update first so we get RequiredBones up to date.
 	Super::RefreshBoneTransforms(NULL); // Pass NULL so we force non threaded work
 
-	const bool bIsPreviewInstance = (PreviewInstance && PreviewInstance == AnimScriptInstance);
-
-	BakedAnimationPoses.Reset();
-	if(bDisplayBakedAnimation && bIsPreviewInstance && PreviewInstance->GetRequiredBones().IsValid())
+	// none of these code works if we don't have anim instance, so no reason to check it for every if
+	if (AnimScriptInstance && AnimScriptInstance->GetRequiredBones().IsValid())
 	{
-		if(UAnimSequence* Sequence = Cast<UAnimSequence>(PreviewInstance->GetCurrentAsset()))
+		const bool bIsPreviewInstance = (PreviewInstance && PreviewInstance == AnimScriptInstance);	
+		FBoneContainer& BoneContainer = AnimScriptInstance->GetRequiredBones();
+
+		BakedAnimationPoses.Reset();
+		if(bDisplayBakedAnimation && bIsPreviewInstance)
 		{
-			BakedAnimationPoses.AddUninitialized(PreviewInstance->GetRequiredBones().GetNumBones());
-			bool bSavedUseSourceData = PreviewInstance->GetRequiredBones().ShouldUseSourceData();
-			PreviewInstance->GetRequiredBones().SetUseRAWData(true);
-			PreviewInstance->GetRequiredBones().SetUseSourceData(false);
-			PreviewInstance->EnableControllers(false);
-			GenSpaceBases(BakedAnimationPoses);
-			PreviewInstance->GetRequiredBones().SetUseRAWData(false);
-			PreviewInstance->GetRequiredBones().SetUseSourceData(bSavedUseSourceData);
-			PreviewInstance->EnableControllers(true);
+			if(UAnimSequence* Sequence = Cast<UAnimSequence>(PreviewInstance->GetCurrentAsset()))
+			{
+				BakedAnimationPoses.AddUninitialized(BoneContainer.GetNumBones());
+				bool bSavedUseSourceData = BoneContainer.ShouldUseSourceData();
+				BoneContainer.SetUseRAWData(true);
+				BoneContainer.SetUseSourceData(false);
+				PreviewInstance->EnableControllers(false);
+				GenSpaceBases(BakedAnimationPoses);
+				BoneContainer.SetUseRAWData(false);
+				BoneContainer.SetUseSourceData(bSavedUseSourceData);
+				PreviewInstance->EnableControllers(true);
+			}
 		}
-	}
 
-	SourceAnimationPoses.Reset();
-	if(bDisplaySourceAnimation && bIsPreviewInstance && PreviewInstance->GetRequiredBones().IsValid())
-	{
-		if(UAnimSequence* Sequence = Cast<UAnimSequence>(PreviewInstance->GetCurrentAsset()))
+		SourceAnimationPoses.Reset();
+		if(bDisplaySourceAnimation && bIsPreviewInstance)
 		{
-			SourceAnimationPoses.AddUninitialized(PreviewInstance->GetRequiredBones().GetNumBones());
-			bool bSavedUseSourceData = PreviewInstance->GetRequiredBones().ShouldUseSourceData();
-			PreviewInstance->GetRequiredBones().SetUseSourceData(true);
-			PreviewInstance->EnableControllers(false);
-			GenSpaceBases(SourceAnimationPoses);
-			PreviewInstance->GetRequiredBones().SetUseSourceData(bSavedUseSourceData);
-			PreviewInstance->EnableControllers(true);
+			if(UAnimSequence* Sequence = Cast<UAnimSequence>(PreviewInstance->GetCurrentAsset()))
+			{
+				SourceAnimationPoses.AddUninitialized(BoneContainer.GetNumBones());
+				bool bSavedUseSourceData = BoneContainer.ShouldUseSourceData();
+				BoneContainer.SetUseSourceData(true);
+				PreviewInstance->EnableControllers(false);
+				GenSpaceBases(SourceAnimationPoses);
+				BoneContainer.SetUseSourceData(bSavedUseSourceData);
+				PreviewInstance->EnableControllers(true);
+			}
 		}
-	}
 
-	UncompressedSpaceBases.Reset();
-	if (bDisplayRawAnimation && AnimScriptInstance && AnimScriptInstance->GetRequiredBones().IsValid())
-	{
-		UncompressedSpaceBases.AddUninitialized(AnimScriptInstance->GetRequiredBones().GetNumBones());
+		UncompressedSpaceBases.Reset();
+		if ( bDisplayRawAnimation )
+		{
+			UncompressedSpaceBases.AddUninitialized(BoneContainer.GetNumBones());
 
-		AnimScriptInstance->GetRequiredBones().SetUseRAWData(true);
-		GenSpaceBases(UncompressedSpaceBases);
-		AnimScriptInstance->GetRequiredBones().SetUseRAWData(false);
-	}
+			BoneContainer.SetUseRAWData(true);
+			GenSpaceBases(UncompressedSpaceBases);
+			BoneContainer.SetUseRAWData(false);
+		}
 
-	// Non retargeted pose.
-	NonRetargetedSpaceBases.Reset();
-	if( bDisplayNonRetargetedPose && AnimScriptInstance && AnimScriptInstance->GetRequiredBones().IsValid() )
-	{
-		NonRetargetedSpaceBases.AddUninitialized(AnimScriptInstance->GetRequiredBones().GetNumBones());
-		AnimScriptInstance->GetRequiredBones().SetDisableRetargeting(true);
-		GenSpaceBases(NonRetargetedSpaceBases);
-		AnimScriptInstance->GetRequiredBones().SetDisableRetargeting(false);
-	}
+		// Non retargeted pose.
+		NonRetargetedSpaceBases.Reset();
+		if( bDisplayNonRetargetedPose )
+		{
+			NonRetargetedSpaceBases.AddUninitialized(BoneContainer.GetNumBones());
+			BoneContainer.SetDisableRetargeting(true);
+			GenSpaceBases(NonRetargetedSpaceBases);
+			BoneContainer.SetDisableRetargeting(false);
+		}
 
-	// Only works in PreviewInstance, and not for anim blueprint. This is intended.
-	AdditiveBasePoses.Reset();
-	if( bDisplayAdditiveBasePose && bIsPreviewInstance && PreviewInstance->GetRequiredBones().IsValid() )
-	{
-		if (UAnimSequence* Sequence = Cast<UAnimSequence>(PreviewInstance->GetCurrentAsset())) 
-		{ 
-			if (Sequence->IsValidAdditive()) 
+		// Only works in PreviewInstance, and not for anim blueprint. This is intended.
+		AdditiveBasePoses.Reset();
+		if( bDisplayAdditiveBasePose && bIsPreviewInstance )
+		{
+			if (UAnimSequence* Sequence = Cast<UAnimSequence>(PreviewInstance->GetCurrentAsset())) 
 			{ 
-				FCSPose<FCompactPose> CSAdditiveBasePose;
-				{
-					FCompactPose AdditiveBasePose;
-					FBlendedCurve AdditiveCurve;
-					AdditiveCurve.InitFrom(AnimScriptInstance->GetSkelMeshComponent()->GetCachedAnimCurveMappingNameUids());
-					AdditiveBasePose.SetBoneContainer(&PreviewInstance->GetRequiredBones());
-					Sequence->GetAdditiveBasePose(AdditiveBasePose, AdditiveCurve, FAnimExtractContext(PreviewInstance->GetCurrentTime()));
-					CSAdditiveBasePose.InitPose(AdditiveBasePose);
-				}
-
-				FBoneContainer& BoneContainer = PreviewInstance->GetRequiredBones();
-				const int32 NumSkeletonBones = BoneContainer.GetNumBones();
+				if (Sequence->IsValidAdditive()) 
+				{ 
+					FCSPose<FCompactPose> CSAdditiveBasePose;
+					{
+						FCompactPose AdditiveBasePose;
+						FBlendedCurve AdditiveCurve;
+						AdditiveCurve.InitFrom(BoneContainer);
+						AdditiveBasePose.SetBoneContainer(&BoneContainer);
+						Sequence->GetAdditiveBasePose(AdditiveBasePose, AdditiveCurve, FAnimExtractContext(PreviewInstance->GetCurrentTime()));
+						CSAdditiveBasePose.InitPose(AdditiveBasePose);
+					}
 
-				AdditiveBasePoses.AddUninitialized(NumSkeletonBones);
+					const int32 NumSkeletonBones = BoneContainer.GetNumBones();
 
-				for (int32 i = 0; i < AdditiveBasePoses.Num(); ++i)
-				{
-					FCompactPoseBoneIndex CompactIndex = BoneContainer.MakeCompactPoseIndex(FMeshPoseBoneIndex(i));
+					AdditiveBasePoses.AddUninitialized(NumSkeletonBones);
 
-					// AdditiveBasePoses has one entry for every bone in the asset ref skeleton - if we're on a LOD
-					// we need to check this is actually valid for the current pose.
-					if(CSAdditiveBasePose.GetPose().IsValidIndex(CompactIndex))
-					{
-						AdditiveBasePoses[i] = CSAdditiveBasePose.GetComponentSpaceTransform(CompactIndex);
-					}
-					else
+					for (int32 i = 0; i < AdditiveBasePoses.Num(); ++i)
 					{
-						AdditiveBasePoses[i] = FTransform::Identity;
+						FCompactPoseBoneIndex CompactIndex = BoneContainer.MakeCompactPoseIndex(FMeshPoseBoneIndex(i));
+
+						// AdditiveBasePoses has one entry for every bone in the asset ref skeleton - if we're on a LOD
+						// we need to check this is actually valid for the current pose.
+						if(CSAdditiveBasePose.GetPose().IsValidIndex(CompactIndex))
+						{
+							AdditiveBasePoses[i] = CSAdditiveBasePose.GetComponentSpaceTransform(CompactIndex);
+						}
+						else
+						{
+							AdditiveBasePoses[i] = FTransform::Identity;
+						}
 					}
 				}
 			}
diff --git a/Engine/Source/Editor/UnrealEd/Private/ApexClothingUtils.cpp b/Engine/Source/Editor/UnrealEd/Private/ApexClothingUtils.cpp
--- a/Engine/Source/Editor/UnrealEd/Private/ApexClothingUtils.cpp
+++ b/Engine/Source/Editor/UnrealEd/Private/ApexClothingUtils.cpp
@@ -976,7 +976,7 @@ bool AssociateClothingAssetWithSkeletalMesh(USkeletalMesh* SkelMesh, int32 LODIn
 	LODModel.MultiSizeIndexContainer.GetIndexBuffer(OutIndexBuffer);
 
 	ClothSection.BaseIndex = OutIndexBuffer.Num();
-	ClothSection.NumTriangles = TriangleCount;
+	ClothSection.NumTriangles = TempOrigSection.NumTriangles;
 
 	uint32 NumVertices = 0;
 	for(int32 SectionIdx = 0; SectionIdx < NewSectionIndex; SectionIdx++)
diff --git a/Engine/Source/Editor/UnrealEd/Private/AssetEditorModeManager.cpp b/Engine/Source/Editor/UnrealEd/Private/AssetEditorModeManager.cpp
--- a/Engine/Source/Editor/UnrealEd/Private/AssetEditorModeManager.cpp
+++ b/Engine/Source/Editor/UnrealEd/Private/AssetEditorModeManager.cpp
@@ -59,3 +59,7 @@ void FAssetEditorModeManager::SetPreviewScene(class FPreviewScene* NewPreviewSce
 	PreviewScene = NewPreviewScene;
 }
 
+FPreviewScene* FAssetEditorModeManager::GetPreviewScene() const
+{
+	return PreviewScene;
+}
\ No newline at end of file
diff --git a/Engine/Source/Editor/UnrealEd/Private/AssetNotifications.cpp b/Engine/Source/Editor/UnrealEd/Private/AssetNotifications.cpp
--- a/Engine/Source/Editor/UnrealEd/Private/AssetNotifications.cpp
+++ b/Engine/Source/Editor/UnrealEd/Private/AssetNotifications.cpp
@@ -7,7 +7,7 @@
 
 #define LOCTEXT_NAMESPACE "AssetNotifications"
 
-void FAssetNotifications::SkeletonNeedsToBeSaved(USkeleton* Skeleton)
+void FAssetNotifications::SkeletonNeedsToBeSaved(const USkeleton* Skeleton)
 {
 	FFormatNamedArguments Args;
 	Args.Add( TEXT("SkeletonName"), FText::FromString( Skeleton->GetName() ) );
diff --git a/Engine/Source/Editor/UnrealEd/Private/AssetViewerSettings.cpp b/Engine/Source/Editor/UnrealEd/Private/AssetViewerSettings.cpp
--- a/Engine/Source/Editor/UnrealEd/Private/AssetViewerSettings.cpp
+++ b/Engine/Source/Editor/UnrealEd/Private/AssetViewerSettings.cpp
@@ -11,14 +11,26 @@ UAssetViewerSettings::UAssetViewerSettings()
 	// Make sure there always is one profile as default
 	Profiles.AddDefaulted(1);
 	Profiles[0].ProfileName = TEXT("Profile_0");
-
-	NumProfiles = Profiles.Num();
+	NumProfiles = 1;
 }
 
 UAssetViewerSettings* UAssetViewerSettings::Get()
 {
 	// This is a singleton, use default object
-	UAssetViewerSettings* DefaultSettings = GetMutableDefault<UAssetViewerSettings>();	
+	UAssetViewerSettings* DefaultSettings = GetMutableDefault<UAssetViewerSettings>();
+
+	// Load environment map textures (once)
+	static bool bInitialized = false;
+	if (!bInitialized)
+	{
+		for (FPreviewSceneProfile& Profile : DefaultSettings->Profiles)
+		{
+			Profile.LoadEnvironmentMap();
+		}
+
+		bInitialized = true;
+	}
+
 	return DefaultSettings;
 }
 
@@ -37,21 +49,11 @@ void UAssetViewerSettings::PostEditChangeProperty(struct FPropertyChangedEvent&
 		PropertyName = GET_MEMBER_NAME_CHECKED(FPreviewSceneProfile, PostProcessingSettings);
 	}
 
-	// Check for identical names (temporary until we validate in the property settings view text input box)
-	if (PropertyName == GET_MEMBER_NAME_CHECKED(FPreviewSceneProfile, ProfileName))
-	{	
-		for (int32 ProfileIndex = 0; ProfileIndex < Profiles.Num(); ++ProfileIndex)
-		{
-			for (int32 CheckIndex = 0; CheckIndex < Profiles.Num(); ++CheckIndex)
-			{
-				// Simply append a string if two profile names are the same
-				if (ProfileIndex != CheckIndex && Profiles[ProfileIndex].ProfileName == Profiles[CheckIndex].ProfileName)
-				{
-					FString ProfileName = Profiles[CheckIndex].ProfileName;
-					Profiles[CheckIndex].ProfileName += "_duplicatename";
-				}
-			}			
-		}
+	// Store path to the set enviroment map texture
+	if (PropertyName == GET_MEMBER_NAME_CHECKED(FPreviewSceneProfile, EnvironmentCubeMap))
+	{
+		int32 ProfileIndex = GetMutableDefault<UEditorPerProjectUserSettings>()->AssetViewerProfileIndex;
+		Profiles[ProfileIndex].EnvironmentCubeMapPath = Profiles[ProfileIndex].EnvironmentCubeMap.ToString();
 	}
 
 	if (NumProfiles != Profiles.Num())
@@ -74,4 +76,3 @@ void UAssetViewerSettings::PostInitProperties()
 		ProjectSettings->AssetViewerProfileIndex = 0;
 	}
 }
-
diff --git a/Engine/Source/Editor/UnrealEd/Private/Editor.cpp b/Engine/Source/Editor/UnrealEd/Private/Editor.cpp
--- a/Engine/Source/Editor/UnrealEd/Private/Editor.cpp
+++ b/Engine/Source/Editor/UnrealEd/Private/Editor.cpp
@@ -484,14 +484,30 @@ void FReimportManager::ValidateAllSourceFileAndReimport(TArray<UObject*> &ToImpo
 		//If user ignore those asset just not add them to CopyOfSelectedAssets
 	}
 
-	for (auto Asset : CopyOfSelectedAssets)
+	FReimportManager::Instance()->ReimportMultiple(CopyOfSelectedAssets, /*bAskForNewFileIfMissing=*/false);
+}
+
+bool FReimportManager::ReimportMultiple(TArrayView<UObject*> Objects, bool bAskForNewFileIfMissing /*= false*/, bool bShowNotification /*= true*/, FString PreferedReimportFile /*= TEXT("")*/, FReimportHandler* SpecifiedReimportHandler /*= nullptr */)
+{
+	bool bBulkSuccess = true;
+
+	FScopedSlowTask BulkReimportTask((float)Objects.Num(), LOCTEXT("BulkReimport_Title", "Reimporting..."));
+
+	for(UObject* CurrentObject : Objects)
 	{
-		if (Asset)
+		if(CurrentObject)
 		{
-			//We already ask for new files
-			this->Reimport(Asset, /*bAskForNewFileIfMissing=*/false);
+			FText SingleTaskTest = FText::Format(LOCTEXT("BulkReimport_SingleItem", "Reimporting {0}"), FText::FromString(CurrentObject->GetName()));
+			FScopedSlowTask SingleObjectTask(1.0f, SingleTaskTest);
+			SingleObjectTask.EnterProgressFrame(1.0f);
+
+			bBulkSuccess = bBulkSuccess && Reimport(CurrentObject, bAskForNewFileIfMissing, bShowNotification, PreferedReimportFile, SpecifiedReimportHandler);
 		}
+
+		BulkReimportTask.EnterProgressFrame(1.0f);
 	}
+
+	return bBulkSuccess;
 }
 
 void FReimportManager::GetNewReimportPath(UObject* Obj, TArray<FString>& InOutFilenames)
diff --git a/Engine/Source/Editor/UnrealEd/Private/EditorEngine.cpp b/Engine/Source/Editor/UnrealEd/Private/EditorEngine.cpp
--- a/Engine/Source/Editor/UnrealEd/Private/EditorEngine.cpp
+++ b/Engine/Source/Editor/UnrealEd/Private/EditorEngine.cpp
@@ -1979,7 +1979,7 @@ UAudioComponent* UEditorEngine::ResetPreviewAudioComponent( USoundBase* Sound, U
 			PreviewSoundCue = NewObject<USoundCue>();
 			// Set world to NULL as it will most likely become invalid in the next PIE/Simulate session and the
 			// component will be left with invalid pointer.
-			PreviewAudioComponent = FAudioDevice::CreateComponent(PreviewSoundCue, NULL, NULL, false);
+			PreviewAudioComponent = FAudioDevice::CreateComponent(PreviewSoundCue);
 		}
 
 		check(PreviewAudioComponent);
diff --git a/Engine/Source/Editor/UnrealEd/Private/Factories/ActorFactory.cpp b/Engine/Source/Editor/UnrealEd/Private/Factories/ActorFactory.cpp
--- a/Engine/Source/Editor/UnrealEd/Private/Factories/ActorFactory.cpp
+++ b/Engine/Source/Editor/UnrealEd/Private/Factories/ActorFactory.cpp
@@ -1114,8 +1114,7 @@ static UBillboardComponent* CreateEditorOnlyBillboardComponent(AActor* ActorOwne
 	BillboardComponent->Sprite = LoadObject<UTexture2D>(nullptr, TEXT("/Engine/EditorResources/EmptyActor.EmptyActor"));
 	BillboardComponent->RelativeScale3D = FVector(0.5f, 0.5f, 0.5f);
 	BillboardComponent->Mobility = EComponentMobility::Movable;
-	BillboardComponent->AlwaysLoadOnClient = false;
-	BillboardComponent->AlwaysLoadOnServer = false;
+	BillboardComponent->bIsEditorOnly = true;
 
 	BillboardComponent->SetupAttachment(AttachParent);
 
diff --git a/Engine/Source/Editor/UnrealEd/Private/Factories/EditorFactories.cpp b/Engine/Source/Editor/UnrealEd/Private/Factories/EditorFactories.cpp
--- a/Engine/Source/Editor/UnrealEd/Private/Factories/EditorFactories.cpp
+++ b/Engine/Source/Editor/UnrealEd/Private/Factories/EditorFactories.cpp
@@ -115,6 +115,8 @@
 #include "Sound/SoundNodeDialoguePlayer.h"
 #include "Factories/CanvasRenderTarget2DFactoryNew.h"
 #include "ImageUtils.h"
+#include "Engine/PreviewMeshCollection.h"
+#include "Factories/PreviewMeshCollectionFactory.h"
 
 DEFINE_LOG_CATEGORY_STATIC(LogEditorFactories, Log, All);
 
@@ -2501,7 +2503,6 @@ UObject* USoundConcurrencyFactory::FactoryCreateNew(UClass* Class, UObject* InPa
 	return NewObject<USoundConcurrency>(InParent, Name, Flags);
 }
 
-
 /*------------------------------------------------------------------------------
 	UParticleSystemFactoryNew.
 ------------------------------------------------------------------------------*/
@@ -7510,6 +7511,55 @@ UObject* UDataTableFactory::FactoryCreateNew(UClass* Class, UObject* InParent, F
 	return DataTable;
 }
 
+/*------------------------------------------------------------------------------
+ UPreviewMeshCollectionFactory implementation.
+------------------------------------------------------------------------------*/
+
+UPreviewMeshCollectionFactory::UPreviewMeshCollectionFactory(const FObjectInitializer& ObjectInitializer)
+	: Super(ObjectInitializer)
+{
+	SupportedClass = UPreviewMeshCollection::StaticClass();
+	bCreateNew = true;
+}
+
+FText UPreviewMeshCollectionFactory::GetDisplayName() const
+{
+	return LOCTEXT("PreviewMeshCollection", "Preview Mesh Collection");
+}
+
+FText UPreviewMeshCollectionFactory::GetToolTip() const
+{
+	return LOCTEXT("PreviewMeshCollection_Tooltip", "Preview Mesh Collections are used to build collections of related skeletal meshes that are animated together (such as components of a character)");
+}
+
+bool UPreviewMeshCollectionFactory::ConfigureProperties()
+{
+	if (CurrentSkeleton.IsValid())
+	{
+		return true;
+	}
+
+	USkeleton* Skeleton = ChooseSkeleton();
+	if (Skeleton != nullptr)
+	{
+		CurrentSkeleton = Skeleton;
+		return true;
+	}
+
+	return false;
+}
+
+UObject* UPreviewMeshCollectionFactory::FactoryCreateNew(UClass* Class, UObject* InParent, FName Name, EObjectFlags Flags, UObject* Context, FFeedbackContext* Warn)
+{
+	UPreviewMeshCollection* NewCollection = nullptr;
+	if (CurrentSkeleton.IsValid())
+	{
+		NewCollection = NewObject<UPreviewMeshCollection>(InParent, Name, Flags);
+		NewCollection->Skeleton = CurrentSkeleton.Get();
+	}
+
+	return NewCollection;
+}
 
 #undef LOCTEXT_NAMESPACE
 
diff --git a/Engine/Source/Editor/UnrealEd/Private/Fbx/FbxSkeletalMeshImport.cpp b/Engine/Source/Editor/UnrealEd/Private/Fbx/FbxSkeletalMeshImport.cpp
--- a/Engine/Source/Editor/UnrealEd/Private/Fbx/FbxSkeletalMeshImport.cpp
+++ b/Engine/Source/Editor/UnrealEd/Private/Fbx/FbxSkeletalMeshImport.cpp
@@ -1530,7 +1530,9 @@ USkeletalMesh* UnFbx::FFbxImporter::ImportSkeletalMesh(UObject* InParent, TArray
 			if(bFirstMesh || (LastMergeBonesChoice != EAppReturnType::NoAll && LastMergeBonesChoice != EAppReturnType::YesAll))
 			{
 				LastMergeBonesChoice = FMessageDialog::Open(EAppMsgType::YesNoYesAllNoAllCancel,
-															LOCTEXT("SkeletonFailed_BoneMerge", "FAILED TO MERGE BONES:\n\n This could happen if significant hierarchical change has been made\n - i.e. inserting bone between nodes\n Would you like to regenerate Skeleton from this mesh? \n\n ***WARNING: THIS WILL REQUIRE RECOMPRESS ALL ANIMATION DATA AND POTENTIALLY INVALIDATE***\n"));
+															LOCTEXT("SkeletonFailed_BoneMerge", "FAILED TO MERGE BONES:\n\n This could happen if significant hierarchical change has been made\n"
+																"- i.e. inserting bone between nodes\n Would you like to regenerate Skeleton from this mesh? \n\n"
+																"***WARNING: THIS WILL REQUIRE RECOMPRESS ALL ANIMATION DATA AND POTENTIALLY INVALIDATE***\n"));
 				bToastSaveMessage = true;
 			}
 			
@@ -1549,6 +1551,41 @@ USkeletalMesh* UnFbx::FFbxImporter::ImportSkeletalMesh(UObject* InParent, TArray
 			{
 				if ( Skeleton->RecreateBoneTree( SkeletalMesh ) && bToastSaveMessage)
 				{
+					// @todo: this is a lot of message box but this requires user input and it can be very annoying to miss
+					// make sure to go through all skeletalmesh and merge them also to recreate the issue. 
+					if (FMessageDialog::Open(EAppMsgType::YesNo,
+						LOCTEXT("Skeleton_ReAddAllMeshes", "Would you like to merge all SkeletalMeshes using this skeleton to ensure all bones are merged? This will require to load those SkeletalMeshes.")) == EAppReturnType::Yes)
+					{
+						FAssetRegistryModule& AssetRegistryModule = FModuleManager::LoadModuleChecked<FAssetRegistryModule>(TEXT("AssetRegistry"));
+						TArray<FAssetData> SkeletalMeshAssetData;
+						
+						FARFilter ARFilter;
+						ARFilter.ClassNames.Add(*USkeletalMesh::StaticClass()->GetName());
+						FString& Value = ARFilter.TagsAndValues.Add(TEXT("Skeleton"));
+						Value = FAssetData(Skeleton).GetExportTextName();
+
+						IAssetRegistry& AssetRegistry = AssetRegistryModule.Get();
+						if (AssetRegistry.GetAssets(ARFilter, SkeletalMeshAssetData))
+						{
+							// look through all skeletalmeshes that uses this skeleton
+							for (int32 AssetId = 0; AssetId < SkeletalMeshAssetData.Num(); ++AssetId)
+							{
+								FAssetData& CurAssetData = SkeletalMeshAssetData[AssetId];
+								const USkeletalMesh* ExtraSkeletalMesh = Cast<USkeletalMesh>(CurAssetData.GetAsset());
+								if (SkeletalMesh != ExtraSkeletalMesh && ExtraSkeletalMesh && ExtraSkeletalMesh->IsPendingKill() == false)
+								{
+									// merge still can fail, then print message box
+									if (Skeleton->MergeAllBonesToBoneTree(ExtraSkeletalMesh) == false)
+									{
+										// print warning
+										FMessageDialog::Open(EAppMsgType::Ok,
+											FText::Format(LOCTEXT("SkeletonRegenError_RemergingBones", "Failed to merge SkeletalMesh '{0}'."), FText::FromString(ExtraSkeletalMesh->GetName())));
+									}
+								}
+							}
+						}
+					}
+
 					FAssetNotifications::SkeletonNeedsToBeSaved(Skeleton);
 				}
 			}
@@ -1556,7 +1593,7 @@ USkeletalMesh* UnFbx::FFbxImporter::ImportSkeletalMesh(UObject* InParent, TArray
 		else
 		{
 			// ask if they'd like to update their position form this mesh
-			if ( ImportOptions->SkeletonForAnimation && ImportOptions->bUpdateSkeletonReferencePose ) 
+			if ( ImportOptions->SkeletonForAnimation && ImportOptions->bUpdateSkeletonReferencePose )
 			{
 				Skeleton->UpdateReferencePoseFromMesh(SkeletalMesh);
 				FAssetNotifications::SkeletonNeedsToBeSaved(Skeleton);
