- Redoing a Linux compile error fix

diff --git a/Engine/Source/Editor/SkeletonEditor/Private/SkeletonTreeBoneItem.h b/Engine/Source/Editor/SkeletonEditor/Private/SkeletonTreeBoneItem.h
--- a/Engine/Source/Editor/SkeletonEditor/Private/SkeletonTreeBoneItem.h
+++ b/Engine/Source/Editor/SkeletonEditor/Private/SkeletonTreeBoneItem.h
@@ -15,6 +15,8 @@ public:
 		, bWeightedBone(false)
 		, bRequiredBone(false)
 	{}
+	virtual ~FSkeletonTreeBoneItem()
+	{}
 
 	/** ISkeletonTreeItem interface */
 	virtual TSharedRef<ITableRow> MakeTreeRowWidget(const TSharedRef<STableViewBase>& OwnerTable, const TAttribute<FText>& FilterText) override;
