- Fixing another virtual destructor for clang

diff --git a/Engine/Source/Editor/SkeletonEditor/Private/SkeletonTreeSocketItem.h b/Engine/Source/Editor/SkeletonEditor/Private/SkeletonTreeSocketItem.h
--- a/Engine/Source/Editor/SkeletonEditor/Private/SkeletonTreeSocketItem.h
+++ b/Engine/Source/Editor/SkeletonEditor/Private/SkeletonTreeSocketItem.h
@@ -17,6 +17,8 @@ public:
 		, bIsCustomized(bInIsCustomized)
 		, bInlineEditorExpanded(false)
 	{}
+	virtual ~FSkeletonTreeSocketItem()
+	{}
 
 	/** ISkeletonTreeItem interface */
 	virtual TSharedRef<ITableRow> MakeTreeRowWidget(const TSharedRef<STableViewBase>& InOwnerTable, const TAttribute<FText>& InFilterText) override;
