Fix for compiling with clang 3.9.x (UE-38788, missing destructor).

(Merging 3204098 from Orion-Staging to Main).

diff --git a/Engine/Source/Editor/SkeletonEditor/Public/ISkeletonTreeItem.h b/Engine/Source/Editor/SkeletonEditor/Public/ISkeletonTreeItem.h
--- a/Engine/Source/Editor/SkeletonEditor/Public/ISkeletonTreeItem.h
+++ b/Engine/Source/Editor/SkeletonEditor/Public/ISkeletonTreeItem.h
@@ -39,6 +39,8 @@ public:
 		return IsOfTypeByName(TType::GetTypeId());
 	}
 
+	virtual ~ISkeletonTreeItem() {};
+
 	/** Builds the table row widget to display this info */
 	virtual TSharedRef<ITableRow> MakeTreeRowWidget(const TSharedRef<STableViewBase>& InOwnerTable, const TAttribute<FText>& InFilterText) = 0;
 
