Duplicate CL #3233861:
Fix for SkeletonTree not rendering bone tree after previous changes.

diff --git a/Engine/Source/Editor/SkeletonEditor/Private/SSkeletonTree.cpp b/Engine/Source/Editor/SkeletonEditor/Private/SSkeletonTree.cpp
--- a/Engine/Source/Editor/SkeletonEditor/Private/SSkeletonTree.cpp
+++ b/Engine/Source/Editor/SkeletonEditor/Private/SSkeletonTree.cpp
@@ -455,6 +455,9 @@ void SSkeletonTree::CreateTreeColumns()
 		TreeHolder->AddSlot()
 		[
 			SNew(SScrollBorder, SkeletonTreeView.ToSharedRef())
+			[
+				SkeletonTreeView.ToSharedRef()
+			]
 		];
 	}
 
