Copying //UE4/Dev-AnimPhys to //UE4/Dev-Main (Source: //UE4/Dev-AnimPhys @ 3624379)

diff --git a/Engine/Source/Editor/SkeletalMeshEditor/Private/SkeletalMeshEditor.cpp b/Engine/Source/Editor/SkeletalMeshEditor/Private/SkeletalMeshEditor.cpp
--- a/Engine/Source/Editor/SkeletalMeshEditor/Private/SkeletalMeshEditor.cpp
+++ b/Engine/Source/Editor/SkeletalMeshEditor/Private/SkeletalMeshEditor.cpp
@@ -32,6 +32,8 @@
 #include "Application/SlateApplication.h"
 #include "EditorViewportClient.h"
 #include "Settings/EditorExperimentalSettings.h"
+#include "Algo/Transform.h"
+#include "ISkeletonTreeItem.h"
 
 const FName SkeletalMeshEditorAppIdentifier = FName(TEXT("SkeletalMeshEditorApp"));
 
@@ -93,8 +95,8 @@ void FSkeletalMeshEditor::InitSkeletalMeshEditor(const EToolkitMode::Type Mode,
 
 	TSharedPtr<IPersonaPreviewScene> PreviewScene = PersonaToolkit->GetPreviewScene();
 
-	FSkeletonTreeArgs SkeletonTreeArgs(OnPostUndo);
-	SkeletonTreeArgs.OnObjectSelected = FOnObjectSelected::CreateSP(this, &FSkeletalMeshEditor::HandleObjectSelected);
+	FSkeletonTreeArgs SkeletonTreeArgs;
+	SkeletonTreeArgs.OnSelectionChanged = FOnSkeletonTreeSelectionChanged::CreateSP(this, &FSkeletalMeshEditor::HandleSelectionChanged);
 	SkeletonTreeArgs.PreviewScene = PreviewScene;
 
 	ISkeletonEditorModule& SkeletonEditorModule = FModuleManager::GetModuleChecked<ISkeletonEditorModule>("SkeletonEditor");
@@ -295,8 +297,8 @@ void FSkeletalMeshEditor::FillMeshClickMenu(FMenuBuilder& MenuBuilder, HActor* H
 		Action.CanExecuteAction = FCanExecuteAction::CreateSP(this, &FSkeletalMeshEditor::CanApplyClothing, LodIndex, SectionIndex);
 
 		MenuBuilder.AddSubMenu(
-			LOCTEXT("MeshClickMenu_AssetApplyMenu", "Apply Clothing Asset..."),
-			LOCTEXT("MeshClickMenu_AssetApplyMenu_ToolTip", "Select a clothing asset to apply to the selected section."),
+			LOCTEXT("MeshClickMenu_AssetApplyMenu", "Apply Clothing Data..."),
+			LOCTEXT("MeshClickMenu_AssetApplyMenu_ToolTip", "Select clothing data to apply to the selected section."),
 			FNewMenuDelegate::CreateSP(this, &FSkeletalMeshEditor::FillApplyClothingAssetMenu, LodIndex, SectionIndex),
 			Action,
 			TEXT(""),
@@ -307,41 +309,34 @@ void FSkeletalMeshEditor::FillMeshClickMenu(FMenuBuilder& MenuBuilder, HActor* H
 		Action.CanExecuteAction = FCanExecuteAction::CreateSP(this, &FSkeletalMeshEditor::CanRemoveClothing, LodIndex, SectionIndex);
 
 		MenuBuilder.AddMenuEntry(
-			LOCTEXT("MeshClickMenu_RemoveClothing", "Remove Clothing Asset"),
-			LOCTEXT("MeshClickMenu_RemoveClothing_ToolTip", "Remove the currently assigned clothing asset."),
+			LOCTEXT("MeshClickMenu_RemoveClothing", "Remove Clothing Data"),
+			LOCTEXT("MeshClickMenu_RemoveClothing_ToolTip", "Remove the currently assigned clothing data."),
 			FSlateIcon(),
 			Action
 			);
+			
+		Action.ExecuteAction = FExecuteAction();
+		Action.CanExecuteAction = FCanExecuteAction::CreateSP(this, &FSkeletalMeshEditor::CanCreateClothing, LodIndex, SectionIndex);
 
-		const UEditorExperimentalSettings* ExperimentalSettings = GetDefault<UEditorExperimentalSettings>(UEditorExperimentalSettings::StaticClass());
-
-		if(ExperimentalSettings && ExperimentalSettings->bClothingTools)
-		{
-			Action.ExecuteAction = FExecuteAction();
-			Action.CanExecuteAction = FCanExecuteAction::CreateSP(this, &FSkeletalMeshEditor::CanCreateClothing, LodIndex, SectionIndex);
-
-			MenuBuilder.AddSubMenu(
-				LOCTEXT("MeshClickMenu_CreateClothing_Label", "Create Clothing Asset from Section"),
-				LOCTEXT("MeshClickMenu_CreateClothing_ToolTip", "Create a new clothing asset using the selected section as a simulation mesh"),
-				FNewMenuDelegate::CreateSP(this, &FSkeletalMeshEditor::FillCreateClothingMenu, LodIndex, SectionIndex),
-				Action,
-				TEXT(""),
-				EUserInterfaceActionType::Button
-				);
-
-			Action.CanExecuteAction = FCanExecuteAction::CreateSP(this, &FSkeletalMeshEditor::CanCreateClothingLod, LodIndex, SectionIndex);
-
-			MenuBuilder.AddSubMenu(
-				LOCTEXT("MeshClickMenu_CreateClothingNewLod_Label", "Create Clothing LOD from Section"),
-				LOCTEXT("MeshClickMenu_CreateClothingNewLod_ToolTip", "Create a clothing simulation mesh from the selected section and add it as a LOD to an existing clothing asset."),
-				FNewMenuDelegate::CreateSP(this, &FSkeletalMeshEditor::FillCreateClothingLodMenu, LodIndex, SectionIndex),
-				Action,
-				TEXT(""),
-				EUserInterfaceActionType::Button
+		MenuBuilder.AddSubMenu(
+			LOCTEXT("MeshClickMenu_CreateClothing_Label", "Create Clothing Data from Section"),
+			LOCTEXT("MeshClickMenu_CreateClothing_ToolTip", "Create a new clothing data using the selected section as a simulation mesh"),
+			FNewMenuDelegate::CreateSP(this, &FSkeletalMeshEditor::FillCreateClothingMenu, LodIndex, SectionIndex),
+			Action,
+			TEXT(""),
+			EUserInterfaceActionType::Button
 			);
 
-		}
+		Action.CanExecuteAction = FCanExecuteAction::CreateSP(this, &FSkeletalMeshEditor::CanCreateClothingLod, LodIndex, SectionIndex);
 
+		MenuBuilder.AddSubMenu(
+			LOCTEXT("MeshClickMenu_CreateClothingNewLod_Label", "Create Clothing LOD from Section"),
+			LOCTEXT("MeshClickMenu_CreateClothingNewLod_ToolTip", "Create a clothing simulation mesh from the selected section and add it as a LOD to existing clothing data."),
+			FNewMenuDelegate::CreateSP(this, &FSkeletalMeshEditor::FillCreateClothingLodMenu, LodIndex, SectionIndex),
+			Action,
+			TEXT(""),
+			EUserInterfaceActionType::Button
+		);
 	}
 	MenuBuilder.EndSection();
 }
@@ -590,6 +585,16 @@ void FSkeletalMeshEditor::HandleObjectSelected(UObject* InObject)
 	}
 }
 
+void FSkeletalMeshEditor::HandleSelectionChanged(const TArrayView<TSharedPtr<ISkeletonTreeItem>>& InSelectedItems, ESelectInfo::Type InSelectInfo)
+{
+	if (DetailsView.IsValid())
+	{
+		TArray<UObject*> Objects;
+		Algo::TransformIf(InSelectedItems, Objects, [](const TSharedPtr<ISkeletonTreeItem>& InItem) { return InItem->GetObject() != nullptr; }, [](const TSharedPtr<ISkeletonTreeItem>& InItem) { return InItem->GetObject(); });
+		DetailsView->SetObjects(Objects);
+	}
+}
+
 void FSkeletalMeshEditor::PostUndo(bool bSuccess)
 {
 	OnPostUndo.Broadcast();
@@ -657,6 +662,12 @@ void FSkeletalMeshEditor::HandleMeshClick(HActor* HitProxy, const FViewportClick
 		Mesh->SelectedEditorSection = HitProxy->SectionIndex;
 	}
 
+	USkeletalMeshComponent* Component = GetPersonaToolkit()->GetPreviewMeshComponent();
+	if (Component)
+	{
+		Component->PushSelectionToProxy();
+	}
+
 	if(Click.GetKey() == EKeys::RightMouseButton)
 	{
 		FMenuBuilder MenuBuilder(true, nullptr);
diff --git a/Engine/Source/Editor/SkeletalMeshEditor/Private/SkeletalMeshEditor.h b/Engine/Source/Editor/SkeletalMeshEditor/Private/SkeletalMeshEditor.h
--- a/Engine/Source/Editor/SkeletalMeshEditor/Private/SkeletalMeshEditor.h
+++ b/Engine/Source/Editor/SkeletalMeshEditor/Private/SkeletalMeshEditor.h
@@ -9,6 +9,7 @@
 #include "EditorUndoClient.h"
 #include "Toolkits/IToolkitHost.h"
 #include "ISkeletalMeshEditor.h"
+#include "ArrayView.h"
 
 class IDetailsView;
 class IPersonaToolkit;
@@ -16,7 +17,7 @@ class IPersonaViewport;
 class ISkeletonTree;
 class USkeletalMesh;
 class UClothingAssetBase;
-
+class ISkeletonTreeItem;
 struct HActor;
 struct FViewportClick;
 struct FSkeletalMeshClothBuildParams;
@@ -93,6 +94,8 @@ private:
 
 	void HandleObjectSelected(UObject* InObject);
 
+	void HandleSelectionChanged(const TArrayView<TSharedPtr<ISkeletonTreeItem>>& InSelectedItems, ESelectInfo::Type InSelectInfo);
+
 	void HandleReimportMesh();
 
 	/** Callback for toggling UV drawing in the viewport */
diff --git a/Engine/Source/Editor/SkeletalMeshEditor/Private/SkeletalMeshEditorMode.cpp b/Engine/Source/Editor/SkeletalMeshEditor/Private/SkeletalMeshEditorMode.cpp
--- a/Engine/Source/Editor/SkeletalMeshEditor/Private/SkeletalMeshEditorMode.cpp
+++ b/Engine/Source/Editor/SkeletalMeshEditor/Private/SkeletalMeshEditorMode.cpp
@@ -23,7 +23,11 @@ FSkeletalMeshEditorMode::FSkeletalMeshEditorMode(TSharedRef<FWorkflowCentricAppl
 
 	FPersonaModule& PersonaModule = FModuleManager::LoadModuleChecked<FPersonaModule>("Persona");
 	TabFactories.RegisterFactory(PersonaModule.CreateDetailsTabFactory(InHostingApp, FOnDetailsCreated::CreateSP(&SkeletalMeshEditor.Get(), &FSkeletalMeshEditor::HandleDetailsCreated)));
-	TabFactories.RegisterFactory(PersonaModule.CreatePersonaViewportTabFactory(InHostingApp, InSkeletonTree, SkeletalMeshEditor->GetPersonaToolkit()->GetPreviewScene(), SkeletalMeshEditor->OnPostUndo, nullptr, FOnViewportCreated(), true, true));
+
+	FPersonaViewportArgs ViewportArgs(InSkeletonTree, SkeletalMeshEditor->GetPersonaToolkit()->GetPreviewScene(), SkeletalMeshEditor->OnPostUndo);
+
+	TabFactories.RegisterFactory(PersonaModule.CreatePersonaViewportTabFactory(InHostingApp, ViewportArgs));
+
 	TabFactories.RegisterFactory(PersonaModule.CreateAdvancedPreviewSceneTabFactory(InHostingApp, SkeletalMeshEditor->GetPersonaToolkit()->GetPreviewScene()));
 	TabFactories.RegisterFactory(PersonaModule.CreateAssetDetailsTabFactory(InHostingApp, FOnGetAsset::CreateSP(&SkeletalMeshEditor.Get(), &FSkeletalMeshEditor::HandleGetAsset), FOnDetailsCreated::CreateSP(&SkeletalMeshEditor.Get(), &FSkeletalMeshEditor::HandleMeshDetailsCreated)));
 	TabFactories.RegisterFactory(PersonaModule.CreateMorphTargetTabFactory(InHostingApp, SkeletalMeshEditor->GetPersonaToolkit()->GetPreviewScene(), SkeletalMeshEditor->OnPostUndo));
diff --git a/Engine/Source/Editor/SkeletonEditor/Private/BoneProxy.cpp b/Engine/Source/Editor/SkeletonEditor/Private/BoneProxy.cpp
--- a/Engine/Source/Editor/SkeletonEditor/Private/BoneProxy.cpp
+++ b/Engine/Source/Editor/SkeletonEditor/Private/BoneProxy.cpp
@@ -16,6 +16,7 @@ UBoneProxy::UBoneProxy()
 	, PreviousRotation(FRotator::ZeroRotator)
 	, PreviousScale(FVector::ZeroVector)
 	, bManipulating(false)
+	, bIsTickable(false)
 {
 }
 
@@ -62,7 +63,7 @@ void UBoneProxy::Tick(float DeltaTime)
 
 bool UBoneProxy::IsTickable() const
 {
-	return true;
+	return bIsTickable;
 }
 
 TStatId UBoneProxy::GetStatId() const
@@ -98,66 +99,69 @@ void UBoneProxy::PreEditChange(FEditPropertyChain& PropertyAboutToChange)
 
 void UBoneProxy::PostEditChangeProperty(FPropertyChangedEvent& PropertyChangedEvent)
 {
-	if (UDebugSkelMeshComponent* Component = SkelMeshComponent.Get())
+	if(PropertyChangedEvent.Property != nullptr)
 	{
-		if (Component->PreviewInstance && Component->AnimScriptInstance == Component->PreviewInstance)
+		if (UDebugSkelMeshComponent* Component = SkelMeshComponent.Get())
 		{
-			bManipulating = (PropertyChangedEvent.ChangeType == EPropertyChangeType::Interactive);
-
-			int32 BoneIndex = Component->GetBoneIndex(BoneName);
-			if (BoneIndex != INDEX_NONE && BoneIndex < Component->GetNumComponentSpaceTransforms())
+			if (Component->PreviewInstance && Component->AnimScriptInstance == Component->PreviewInstance)
 			{
-				FTransform BoneTransform = Component->GetBoneTransform(BoneIndex);
-				FMatrix BoneLocalCoordSystem = Component->GetBoneTransform(BoneIndex).ToMatrixNoScale().RemoveTranslation();
-				FAnimNode_ModifyBone& ModifyBone = Component->PreviewInstance->ModifyBone(BoneName);
-				FTransform ModifyBoneTransform(ModifyBone.Rotation, ModifyBone.Translation, ModifyBone.Scale);
-				FTransform BaseTransform = BoneTransform.GetRelativeTransformReverse(ModifyBoneTransform);
+				bManipulating = (PropertyChangedEvent.ChangeType == EPropertyChangeType::Interactive);
 
-				if (PropertyChangedEvent.MemberProperty->GetFName() == GET_MEMBER_NAME_CHECKED(UBoneProxy, Location))
+				int32 BoneIndex = Component->GetBoneIndex(BoneName);
+				if (BoneIndex != INDEX_NONE && BoneIndex < Component->GetNumComponentSpaceTransforms())
 				{
-					FVector Delta = (Location - PreviousLocation);
-					if (!Delta.IsNearlyZero())
+					FTransform BoneTransform = Component->GetBoneTransform(BoneIndex);
+					FMatrix BoneLocalCoordSystem = Component->GetBoneTransform(BoneIndex).ToMatrixNoScale().RemoveTranslation();
+					FAnimNode_ModifyBone& ModifyBone = Component->PreviewInstance->ModifyBone(BoneName);
+					FTransform ModifyBoneTransform(ModifyBone.Rotation, ModifyBone.Translation, ModifyBone.Scale);
+					FTransform BaseTransform = BoneTransform.GetRelativeTransformReverse(ModifyBoneTransform);
+
+					if (PropertyChangedEvent.MemberProperty->GetFName() == GET_MEMBER_NAME_CHECKED(UBoneProxy, Location))
 					{
-						if (bLocalLocation)
+						FVector Delta = (Location - PreviousLocation);
+						if (!Delta.IsNearlyZero())
 						{
-							Delta = BoneLocalCoordSystem.TransformPosition(Delta);
-						}
+							if (bLocalLocation)
+							{
+								Delta = BoneLocalCoordSystem.TransformPosition(Delta);
+							}
 
-						FVector BoneSpaceDelta = BaseTransform.TransformVector(Delta);
-						ModifyBone.Translation += BoneSpaceDelta;
+							FVector BoneSpaceDelta = BaseTransform.TransformVector(Delta);
+							ModifyBone.Translation += BoneSpaceDelta;
+						}
 					}
-				}
-				else if (PropertyChangedEvent.MemberProperty->GetFName() == GET_MEMBER_NAME_CHECKED(UBoneProxy, Rotation))
-				{
-					FRotator Delta = (Rotation - PreviousRotation);
-					if (!Delta.IsNearlyZero())
-					{					
-						if (bLocalRotation)
-						{
-							// get delta in current coord space
-							Delta = (BoneLocalCoordSystem.Inverse() * FRotationMatrix(Delta) * BoneLocalCoordSystem).Rotator();
+					else if (PropertyChangedEvent.MemberProperty->GetFName() == GET_MEMBER_NAME_CHECKED(UBoneProxy, Rotation))
+					{
+						FRotator Delta = (Rotation - PreviousRotation);
+						if (!Delta.IsNearlyZero())
+						{					
+							if (bLocalRotation)
+							{
+								// get delta in current coord space
+								Delta = (BoneLocalCoordSystem.Inverse() * FRotationMatrix(Delta) * BoneLocalCoordSystem).Rotator();
+							}
+
+							FVector RotAxis;
+							float RotAngle;
+							Delta.Quaternion().ToAxisAndAngle(RotAxis, RotAngle);
+
+							FVector4 BoneSpaceAxis = BaseTransform.TransformVectorNoScale(RotAxis);
+
+							//Calculate the new delta rotation
+							FQuat NewDeltaQuat(BoneSpaceAxis, RotAngle);
+							NewDeltaQuat.Normalize();
+
+							FRotator NewRotation = (ModifyBoneTransform * FTransform(NewDeltaQuat)).Rotator();
+							ModifyBone.Rotation = NewRotation;
 						}
-
-						FVector RotAxis;
-						float RotAngle;
-						Delta.Quaternion().ToAxisAndAngle(RotAxis, RotAngle);
-
-						FVector4 BoneSpaceAxis = BaseTransform.TransformVectorNoScale(RotAxis);
-
-						//Calculate the new delta rotation
-						FQuat NewDeltaQuat(BoneSpaceAxis, RotAngle);
-						NewDeltaQuat.Normalize();
-
-						FRotator NewRotation = (ModifyBoneTransform * FTransform(NewDeltaQuat)).Rotator();
-						ModifyBone.Rotation = NewRotation;
 					}
-				}
-				else if (PropertyChangedEvent.MemberProperty->GetFName() == GET_MEMBER_NAME_CHECKED(UBoneProxy, Scale))
-				{
-					FVector Delta = (Scale - PreviousScale);
-					if (!Delta.IsNearlyZero())
+					else if (PropertyChangedEvent.MemberProperty->GetFName() == GET_MEMBER_NAME_CHECKED(UBoneProxy, Scale))
 					{
-						ModifyBone.Scale += Delta;
+						FVector Delta = (Scale - PreviousScale);
+						if (!Delta.IsNearlyZero())
+						{
+							ModifyBone.Scale += Delta;
+						}
 					}
 				}
 			}
diff --git a/Engine/Source/Editor/SkeletonEditor/Private/BoneProxyDetailsCustomization.cpp b/Engine/Source/Editor/SkeletonEditor/Private/BoneProxyDetailsCustomization.cpp
--- a/Engine/Source/Editor/SkeletonEditor/Private/BoneProxyDetailsCustomization.cpp
+++ b/Engine/Source/Editor/SkeletonEditor/Private/BoneProxyDetailsCustomization.cpp
@@ -13,6 +13,7 @@
 #include "MultiBoxBuilder.h"
 #include "SComboButton.h"
 #include "SBox.h"
+#include "Algo/Transform.h"
 
 #define LOCTEXT_NAMESPACE "FBoneProxyDetailsCustomization"
 
@@ -46,83 +47,96 @@ static bool IsWorldTransformChecked(bool* bValuePtr)
 	return !*bValuePtr;
 }
 
-static TSharedRef<SWidget> BuildTransformFieldLabel(bool* bValuePtr, const FText& Label)
+static TSharedRef<SWidget> BuildTransformFieldLabel(bool* bValuePtr, const FText& Label, bool bMultiSelected)
 {
-	FMenuBuilder MenuBuilder( true, nullptr, nullptr);
-
-	FUIAction SetRelativeLocationAction
-	(
-		FExecuteAction::CreateStatic(&OnSetRelativeTransform, bValuePtr),
-		FCanExecuteAction(),
-		FIsActionChecked::CreateStatic(&IsRelativeTransformChecked, bValuePtr)
-	);
-
-	FUIAction SetWorldLocationAction
-	(
-		FExecuteAction::CreateStatic(&OnSetWorldTransform, bValuePtr),
-		FCanExecuteAction(),
-		FIsActionChecked::CreateStatic(&IsWorldTransformChecked, bValuePtr)
-	);
-
-	MenuBuilder.BeginSection( TEXT("TransformType"), FText::Format( LOCTEXT("TransformType", "{0} Type"), Label ) );
-
-	MenuBuilder.AddMenuEntry
-	(
-		FText::Format( LOCTEXT( "LocalLabel", "Local"), Label ),
-		FText::Format( LOCTEXT( "LocalLabel_ToolTip", "{0} is relative to its parent"), Label ),
-		FSlateIcon(),
-		SetRelativeLocationAction,
-		NAME_None, 
-		EUserInterfaceActionType::RadioButton
-	);
-
-	MenuBuilder.AddMenuEntry
-	(
-		FText::Format( LOCTEXT( "WorldLabel", "World"), Label ),
-		FText::Format( LOCTEXT( "WorldLabel_ToolTip", "{0} is relative to the world"), Label ),
-		FSlateIcon(),
-		SetWorldLocationAction,
-		NAME_None,
-		EUserInterfaceActionType::RadioButton
-	);
-
-	MenuBuilder.EndSection();
-
-	return 
-		SNew(SHorizontalBox)
-		+SHorizontalBox::Slot()
-		.HAlign(HAlign_Left)
-		[
-			SNew(SComboButton)
-			.ContentPadding( 0 )
-			.ButtonStyle( FEditorStyle::Get(), "NoBorder" )
-			.ForegroundColor( FSlateColor::UseForeground() )
-			.MenuContent()
-			[
-				MenuBuilder.MakeWidget()
-			]
-			.ButtonContent()
+	if (bMultiSelected)
+	{
+		return SNew(STextBlock)
+			.Text(Label)
+			.Font(IDetailLayoutBuilder::GetDetailFont());
+	}
+	else
+	{
+		FMenuBuilder MenuBuilder(true, nullptr, nullptr);
+
+		FUIAction SetRelativeLocationAction
+		(
+			FExecuteAction::CreateStatic(&OnSetRelativeTransform, bValuePtr),
+			FCanExecuteAction(),
+			FIsActionChecked::CreateStatic(&IsRelativeTransformChecked, bValuePtr)
+		);
+
+		FUIAction SetWorldLocationAction
+		(
+			FExecuteAction::CreateStatic(&OnSetWorldTransform, bValuePtr),
+			FCanExecuteAction(),
+			FIsActionChecked::CreateStatic(&IsWorldTransformChecked, bValuePtr)
+		);
+
+		MenuBuilder.BeginSection( TEXT("TransformType"), FText::Format( LOCTEXT("TransformType", "{0} Type"), Label ) );
+
+		MenuBuilder.AddMenuEntry
+		(
+			FText::Format( LOCTEXT( "LocalLabel", "Local"), Label ),
+			FText::Format( LOCTEXT( "LocalLabel_ToolTip", "{0} is relative to its parent"), Label ),
+			FSlateIcon(),
+			SetRelativeLocationAction,
+			NAME_None, 
+			EUserInterfaceActionType::RadioButton
+		);
+
+		MenuBuilder.AddMenuEntry
+		(
+			FText::Format( LOCTEXT( "WorldLabel", "World"), Label ),
+			FText::Format( LOCTEXT( "WorldLabel_ToolTip", "{0} is relative to the world"), Label ),
+			FSlateIcon(),
+			SetWorldLocationAction,
+			NAME_None,
+			EUserInterfaceActionType::RadioButton
+		);
+
+		MenuBuilder.EndSection();
+
+		return 
+			SNew(SHorizontalBox)
+			+SHorizontalBox::Slot()
+			.HAlign(HAlign_Left)
 			[
-				SNew( SBox )
-				.Padding( FMargin( 0.0f, 0.0f, 2.0f, 0.0f ) )
+				SNew(SComboButton)
+				.ContentPadding( 0 )
+				.ButtonStyle( FEditorStyle::Get(), "NoBorder" )
+				.ForegroundColor( FSlateColor::UseForeground() )
+				.MenuContent()
 				[
-					SNew(STextBlock)
-					.Text_Static(&GetTransformFieldText, bValuePtr, Label)
-					.Font(IDetailLayoutBuilder::GetDetailFont())
+					MenuBuilder.MakeWidget()
 				]
-			]
-		];
+				.ButtonContent()
+				[
+					SNew( SBox )
+					.Padding( FMargin( 0.0f, 0.0f, 2.0f, 0.0f ) )
+					[
+						SNew(STextBlock)
+						.Text_Static(&GetTransformFieldText, bValuePtr, Label)
+						.Font(IDetailLayoutBuilder::GetDetailFont())
+					]
+				]
+			];
+	}
 }
 
 void FBoneProxyDetailsCustomization::CustomizeDetails(IDetailLayoutBuilder& DetailBuilder)
 {
 	TArray<TWeakObjectPtr<UObject>> Objects;
 	DetailBuilder.GetObjectsBeingCustomized(Objects);
-	check(Objects.Num() == 1);
 
-	UBoneProxy* BoneProxy = CastChecked<UBoneProxy>(Objects[0].Get());
+	BoneProxies.Empty();
+	Algo::TransformIf(Objects, BoneProxies, [](TWeakObjectPtr<UObject> InItem) { return InItem.IsValid() && InItem.Get()->IsA<UBoneProxy>(); }, [](TWeakObjectPtr<UObject> InItem) { return CastChecked<UBoneProxy>(InItem.Get()); });
+	TArrayView<UBoneProxy*> BoneProxiesView(BoneProxies);
+
+	UBoneProxy* FirstBoneProxy = CastChecked<UBoneProxy>(Objects[0].Get());
+
 	bool bIsEditingEnabled = true;
-	if (UDebugSkelMeshComponent* Component = BoneProxy->SkelMeshComponent.Get())
+	if (UDebugSkelMeshComponent* Component = FirstBoneProxy->SkelMeshComponent.Get())
 	{
 		bIsEditingEnabled = (Component->AnimScriptInstance == Component->PreviewInstance);
 	}
@@ -138,12 +152,12 @@ void FBoneProxyDetailsCustomization::CustomizeDetails(IDetailLayoutBuilder& Deta
 
 	IDetailPropertyRow& LocationPropertyRow = CategoryBuilder.AddProperty(LocationProperty);
 	LocationPropertyRow.GetDefaultWidgets(NameWidget, ValueWidget);
-	LocationPropertyRow.OverrideResetToDefault(FResetToDefaultOverride::Create(FIsResetToDefaultVisible::CreateSP(this, &FBoneProxyDetailsCustomization::IsResetLocationVisible, BoneProxy), FResetToDefaultHandler::CreateSP(this, &FBoneProxyDetailsCustomization::HandleResetLocation, BoneProxy)));
+	LocationPropertyRow.OverrideResetToDefault(FResetToDefaultOverride::Create(FIsResetToDefaultVisible::CreateSP(this, &FBoneProxyDetailsCustomization::IsResetLocationVisible, BoneProxiesView), FResetToDefaultHandler::CreateSP(this, &FBoneProxyDetailsCustomization::HandleResetLocation, BoneProxiesView)));
 
 	LocationPropertyRow.CustomWidget()
 	.NameContent()
 	[
-		BuildTransformFieldLabel(&BoneProxy->bLocalLocation, LOCTEXT("Location", "Location"))
+		BuildTransformFieldLabel(&FirstBoneProxy->bLocalLocation, LOCTEXT("Location", "Location"), Objects.Num() > 1)
 	]
 	.ValueContent()
 	.MinDesiredWidth(BoneProxyCustomizationConstants::ItemWidth * 3.0f)
@@ -158,12 +172,12 @@ void FBoneProxyDetailsCustomization::CustomizeDetails(IDetailLayoutBuilder& Deta
 
 	IDetailPropertyRow& RotationPropertyRow = CategoryBuilder.AddProperty(RotationProperty);
 	RotationPropertyRow.GetDefaultWidgets(NameWidget, ValueWidget);
-	RotationPropertyRow.OverrideResetToDefault(FResetToDefaultOverride::Create(FIsResetToDefaultVisible::CreateSP(this, &FBoneProxyDetailsCustomization::IsResetRotationVisible, BoneProxy), FResetToDefaultHandler::CreateSP(this, &FBoneProxyDetailsCustomization::HandleResetRotation, BoneProxy)));
+	RotationPropertyRow.OverrideResetToDefault(FResetToDefaultOverride::Create(FIsResetToDefaultVisible::CreateSP(this, &FBoneProxyDetailsCustomization::IsResetRotationVisible, BoneProxiesView), FResetToDefaultHandler::CreateSP(this, &FBoneProxyDetailsCustomization::HandleResetRotation, BoneProxiesView)));
 
 	RotationPropertyRow.CustomWidget()
 	.NameContent()
 	[
-		BuildTransformFieldLabel(&BoneProxy->bLocalRotation, LOCTEXT("Rotation", "Rotation"))
+		BuildTransformFieldLabel(&FirstBoneProxy->bLocalRotation, LOCTEXT("Rotation", "Rotation"), Objects.Num() > 1)
 	]
 	.ValueContent()
 	.MinDesiredWidth(BoneProxyCustomizationConstants::ItemWidth * 3.0f)
@@ -178,7 +192,7 @@ void FBoneProxyDetailsCustomization::CustomizeDetails(IDetailLayoutBuilder& Deta
 
 	IDetailPropertyRow& ScalePropertyRow = CategoryBuilder.AddProperty(ScaleProperty);
 	ScalePropertyRow.GetDefaultWidgets(NameWidget, ValueWidget);
-	ScalePropertyRow.OverrideResetToDefault(FResetToDefaultOverride::Create(FIsResetToDefaultVisible::CreateSP(this, &FBoneProxyDetailsCustomization::IsResetScaleVisible, BoneProxy), FResetToDefaultHandler::CreateSP(this, &FBoneProxyDetailsCustomization::HandleResetScale, BoneProxy)));
+	ScalePropertyRow.OverrideResetToDefault(FResetToDefaultOverride::Create(FIsResetToDefaultVisible::CreateSP(this, &FBoneProxyDetailsCustomization::IsResetScaleVisible, BoneProxiesView), FResetToDefaultHandler::CreateSP(this, &FBoneProxyDetailsCustomization::HandleResetScale, BoneProxiesView)));
 
 	ScalePropertyRow.CustomWidget()
 	.NameContent()
@@ -199,90 +213,117 @@ void FBoneProxyDetailsCustomization::CustomizeDetails(IDetailLayoutBuilder& Deta
 	];
 }
 
-bool FBoneProxyDetailsCustomization::IsResetLocationVisible(TSharedPtr<IPropertyHandle> InPropertyHandle, UBoneProxy* InBoneProxy)
+bool FBoneProxyDetailsCustomization::IsResetLocationVisible(TSharedPtr<IPropertyHandle> InPropertyHandle, TArrayView<UBoneProxy*> InBoneProxies)
 {
-	if (UDebugSkelMeshComponent* Component = InBoneProxy->SkelMeshComponent.Get())
+	for (UBoneProxy* BoneProxy : InBoneProxies)
 	{
-		if (FAnimNode_ModifyBone* ModifyBone = Component->PreviewInstance->FindModifiedBone(InBoneProxy->BoneName))
+		if (UDebugSkelMeshComponent* Component = BoneProxy->SkelMeshComponent.Get())
 		{
-			return ModifyBone->Translation != FVector::ZeroVector;
+			if (FAnimNode_ModifyBone* ModifyBone = Component->PreviewInstance->FindModifiedBone(BoneProxy->BoneName))
+			{
+				if (ModifyBone->Translation != FVector::ZeroVector)
+				{
+					return true;
+				}
+			}
 		}
 	}
 
 	return false;
 }
 
-bool FBoneProxyDetailsCustomization::IsResetRotationVisible(TSharedPtr<IPropertyHandle> InPropertyHandle, UBoneProxy* InBoneProxy)
+bool FBoneProxyDetailsCustomization::IsResetRotationVisible(TSharedPtr<IPropertyHandle> InPropertyHandle, TArrayView<UBoneProxy*> InBoneProxies)
 {
-	if (UDebugSkelMeshComponent* Component = InBoneProxy->SkelMeshComponent.Get())
+	for (UBoneProxy* BoneProxy : InBoneProxies)
 	{
-		if (FAnimNode_ModifyBone* ModifyBone = Component->PreviewInstance->FindModifiedBone(InBoneProxy->BoneName))
+		if (UDebugSkelMeshComponent* Component = BoneProxy->SkelMeshComponent.Get())
 		{
-			return ModifyBone->Rotation != FRotator::ZeroRotator;
+			if (FAnimNode_ModifyBone* ModifyBone = Component->PreviewInstance->FindModifiedBone(BoneProxy->BoneName))
+			{
+				if (ModifyBone->Rotation != FRotator::ZeroRotator)
+				{
+					return true;
+				}
+			}
 		}
 	}
 
 	return false;
 }
 
-bool FBoneProxyDetailsCustomization::IsResetScaleVisible(TSharedPtr<IPropertyHandle> InPropertyHandle, UBoneProxy* InBoneProxy)
+bool FBoneProxyDetailsCustomization::IsResetScaleVisible(TSharedPtr<IPropertyHandle> InPropertyHandle, TArrayView<UBoneProxy*> InBoneProxies)
 {
-	if (UDebugSkelMeshComponent* Component = InBoneProxy->SkelMeshComponent.Get())
+	for (UBoneProxy* BoneProxy : InBoneProxies)
 	{
-		if (FAnimNode_ModifyBone* ModifyBone = Component->PreviewInstance->FindModifiedBone(InBoneProxy->BoneName))
+		if (UDebugSkelMeshComponent* Component = BoneProxy->SkelMeshComponent.Get())
 		{
-			return ModifyBone->Scale != FVector(1.0f);
+			if (FAnimNode_ModifyBone* ModifyBone = Component->PreviewInstance->FindModifiedBone(BoneProxy->BoneName))
+			{
+				if (ModifyBone->Scale != FVector(1.0f))
+				{
+					return true;
+				}
+			}
 		}
 	}
 
 	return false;
 }
 
-void FBoneProxyDetailsCustomization::HandleResetLocation(TSharedPtr<IPropertyHandle> InPropertyHandle, UBoneProxy* InBoneProxy)
+void FBoneProxyDetailsCustomization::HandleResetLocation(TSharedPtr<IPropertyHandle> InPropertyHandle, TArrayView<UBoneProxy*> InBoneProxies)
 {
-	if (UDebugSkelMeshComponent* Component = InBoneProxy->SkelMeshComponent.Get())
-	{
-		FScopedTransaction Transaction(LOCTEXT("ResetLocation", "Reset Location"));
+	FScopedTransaction Transaction(LOCTEXT("ResetLocation", "Reset Location"));
 
-		InBoneProxy->Modify();
-		Component->PreviewInstance->Modify();
+	for (UBoneProxy* BoneProxy : InBoneProxies)
+	{
+		if (UDebugSkelMeshComponent* Component = BoneProxy->SkelMeshComponent.Get())
+		{
+			BoneProxy->Modify();
+			Component->PreviewInstance->Modify();
 
-		FAnimNode_ModifyBone& ModifyBone = Component->PreviewInstance->ModifyBone(InBoneProxy->BoneName);
-		ModifyBone.Translation = FVector::ZeroVector;
+			FAnimNode_ModifyBone& ModifyBone = Component->PreviewInstance->ModifyBone(BoneProxy->BoneName);
+			ModifyBone.Translation = FVector::ZeroVector;
 
-		RemoveUnnecessaryModifications(Component, ModifyBone);
+			RemoveUnnecessaryModifications(Component, ModifyBone);
+		}
 	}
 }
 
-void FBoneProxyDetailsCustomization::HandleResetRotation(TSharedPtr<IPropertyHandle> InPropertyHandle, UBoneProxy* InBoneProxy)
+void FBoneProxyDetailsCustomization::HandleResetRotation(TSharedPtr<IPropertyHandle> InPropertyHandle, TArrayView<UBoneProxy*> InBoneProxies)
 {
-	if (UDebugSkelMeshComponent* Component = InBoneProxy->SkelMeshComponent.Get())
-	{
-		FScopedTransaction Transaction(LOCTEXT("ResetRotation", "Reset Rotation"));
+	FScopedTransaction Transaction(LOCTEXT("ResetRotation", "Reset Rotation"));
 
-		InBoneProxy->Modify();
-		Component->PreviewInstance->Modify();
+	for (UBoneProxy* BoneProxy : InBoneProxies)
+	{
+		if (UDebugSkelMeshComponent* Component = BoneProxy->SkelMeshComponent.Get())
+		{
+			BoneProxy->Modify();
+			Component->PreviewInstance->Modify();
 
-		FAnimNode_ModifyBone& ModifyBone = Component->PreviewInstance->ModifyBone(InBoneProxy->BoneName);
-		ModifyBone.Rotation = FRotator::ZeroRotator;
+			FAnimNode_ModifyBone& ModifyBone = Component->PreviewInstance->ModifyBone(BoneProxy->BoneName);
+			ModifyBone.Rotation = FRotator::ZeroRotator;
 
-		RemoveUnnecessaryModifications(Component, ModifyBone);
+			RemoveUnnecessaryModifications(Component, ModifyBone);
+		}
 	}
 }
 
-void FBoneProxyDetailsCustomization::HandleResetScale(TSharedPtr<IPropertyHandle> InPropertyHandle, UBoneProxy* InBoneProxy)
+void FBoneProxyDetailsCustomization::HandleResetScale(TSharedPtr<IPropertyHandle> InPropertyHandle, TArrayView<UBoneProxy*> InBoneProxies)
 {
-	if (UDebugSkelMeshComponent* Component = InBoneProxy->SkelMeshComponent.Get())
-	{
-		FScopedTransaction Transaction(LOCTEXT("ResetScale", "Reset Scale"));
+	FScopedTransaction Transaction(LOCTEXT("ResetScale", "Reset Scale"));
 
-		InBoneProxy->Modify();
-		Component->PreviewInstance->Modify();
+	for (UBoneProxy* BoneProxy : InBoneProxies)
+	{
+		if (UDebugSkelMeshComponent* Component = BoneProxy->SkelMeshComponent.Get())
+		{
+			BoneProxy->Modify();
+			Component->PreviewInstance->Modify();
 
-		FAnimNode_ModifyBone& ModifyBone = Component->PreviewInstance->ModifyBone(InBoneProxy->BoneName);
-		ModifyBone.Scale = FVector(1.0f);
+			FAnimNode_ModifyBone& ModifyBone = Component->PreviewInstance->ModifyBone(BoneProxy->BoneName);
+			ModifyBone.Scale = FVector(1.0f);
 
-		RemoveUnnecessaryModifications(Component, ModifyBone);
+			RemoveUnnecessaryModifications(Component, ModifyBone);
+		}
 	}
 }
 
diff --git a/Engine/Source/Editor/SkeletonEditor/Private/BoneProxyDetailsCustomization.h b/Engine/Source/Editor/SkeletonEditor/Private/BoneProxyDetailsCustomization.h
--- a/Engine/Source/Editor/SkeletonEditor/Private/BoneProxyDetailsCustomization.h
+++ b/Engine/Source/Editor/SkeletonEditor/Private/BoneProxyDetailsCustomization.h
@@ -3,6 +3,7 @@
 #pragma once
 
 #include "IDetailCustomization.h"
+#include "ArrayView.h"
 
 class IPropertyHandle;
 class UBoneProxy;
@@ -22,23 +23,27 @@ public:
 
 private:
 	/** Handle reset to defaults visibility */
-	bool IsResetLocationVisible(TSharedPtr<IPropertyHandle> InPropertyHandle, UBoneProxy* InBoneProxy);
+	bool IsResetLocationVisible(TSharedPtr<IPropertyHandle> InPropertyHandle, TArrayView<UBoneProxy*> InBoneProxies);
 
 	/** Handle reset to defaults visibility */
-	bool IsResetRotationVisible(TSharedPtr<IPropertyHandle> InPropertyHandle, UBoneProxy* InBoneProxy);
+	bool IsResetRotationVisible(TSharedPtr<IPropertyHandle> InPropertyHandle, TArrayView<UBoneProxy*> InBoneProxies);
 
 	/** Handle reset to defaults visibility */
-	bool IsResetScaleVisible(TSharedPtr<IPropertyHandle> InPropertyHandle, UBoneProxy* InBoneProxy);
+	bool IsResetScaleVisible(TSharedPtr<IPropertyHandle> InPropertyHandle, TArrayView<UBoneProxy*> InBoneProxies);
 
 	/** Handle resetting defaults */
-	void HandleResetLocation(TSharedPtr<IPropertyHandle> InPropertyHandle, UBoneProxy* InBoneProxy);
+	void HandleResetLocation(TSharedPtr<IPropertyHandle> InPropertyHandle, TArrayView<UBoneProxy*> InBoneProxies);
 
 	/** Handle resetting defaults */
-	void HandleResetRotation(TSharedPtr<IPropertyHandle> InPropertyHandle, UBoneProxy* InBoneProxy);
+	void HandleResetRotation(TSharedPtr<IPropertyHandle> InPropertyHandle, TArrayView<UBoneProxy*> InBoneProxies);
 
 	/** Handle resetting defaults */
-	void HandleResetScale(TSharedPtr<IPropertyHandle> InPropertyHandle, UBoneProxy* InBoneProxy);
+	void HandleResetScale(TSharedPtr<IPropertyHandle> InPropertyHandle, TArrayView<UBoneProxy*> InBoneProxies);
 
 	/** Remove any modification node if it has no effect */
 	void RemoveUnnecessaryModifications(UDebugSkelMeshComponent* Component, FAnimNode_ModifyBone& ModifyBone);
+
+private:
+	/** Cached array of selected bone proxies */
+	TArray<UBoneProxy*> BoneProxies;
 };
diff --git a/Engine/Source/Editor/SkeletonEditor/Private/EditableSkeleton.cpp b/Engine/Source/Editor/SkeletonEditor/Private/EditableSkeleton.cpp
--- a/Engine/Source/Editor/SkeletonEditor/Private/EditableSkeleton.cpp
+++ b/Engine/Source/Editor/SkeletonEditor/Private/EditableSkeleton.cpp
@@ -1,6 +1,7 @@
 // Copyright 1998-2017 Epic Games, Inc. All Rights Reserved.
 
 #include "EditableSkeleton.h"
+#include "Editor.h"
 #include "Misc/MessageDialog.h"
 #include "Misc/FeedbackContext.h"
 #include "Modules/ModuleManager.h"
@@ -10,6 +11,7 @@
 #include "EdGraph/EdGraphSchema.h"
 #include "Animation/AnimSequenceBase.h"
 #include "Animation/AnimSequence.h"
+#include "Animation/PoseAsset.h"
 #include "Engine/SkeletalMeshSocket.h"
 #include "Framework/Notifications/NotificationManager.h"
 #include "Widgets/Notifications/SNotificationList.h"
@@ -233,43 +235,112 @@ bool FEditableSkeleton::AddSmartname(const FName& InContainerName, const FName&
 
 void FEditableSkeleton::RenameSmartname(const FName& InContainerName, SmartName::UID_Type InNameUid, const FName& InNewName)
 {
+	FSmartName CurveToRename;
+	if (!Skeleton->GetSmartNameByUID(USkeleton::AnimCurveMappingName, InNameUid, CurveToRename))
+	{
+		return; //Could not get existing smart name
+	}
+
 	const FSmartNameMapping* Mapping = Skeleton->GetSmartNameContainer(InContainerName);
-	if (Mapping)
+	if (!Mapping || Mapping->Exists(InNewName))
+	{
+		return; // Name already exists
+	}
+
+	FText Title = LOCTEXT("RenameCurveDialogTitle", "Confirm Rename");
+	FText ConfirmMessage = LOCTEXT("RenameCurveMessage", "Renaming a curve will necessitate loading and modifying animations and pose assets that use this curve. This could be a slow process.\n\nContinue?");
+
+	if (FMessageDialog::Open(EAppMsgType::YesNo, ConfirmMessage, &Title) == EAppReturnType::Yes)
 	{
-		if (!Mapping->Exists(InNewName))
+		TArray<FAssetData> AnimationAssets;
+
+		TArray<FName> Names = { CurveToRename.DisplayName };
+		GetAssetsContainingCurves(InContainerName, Names, AnimationAssets);
+
+		// AnimationAssets now only contains assets that are using the selected curve(s)
+		if (AnimationAssets.Num() > 0)
+		{
+			TArray<UAnimSequence*> SequencesToRecompress;
+			SequencesToRecompress.Reserve(AnimationAssets.Num());
+
+			// Proceed to delete the curves
+			GWarn->BeginSlowTask(FText::Format(LOCTEXT("RenameCurvesTaskDesc", "Renaming curve for skeleton {0}"), FText::FromString(Skeleton->GetName())), true);
+			FScopedTransaction Transaction(LOCTEXT("RenameCurvesTransactionName", "Rename skeleton curve"));
+
+			// Remove curves from animation assets
+			for (FAssetData& Data : AnimationAssets)
+			{
+				UObject* Asset = Data.GetAsset();
+
+				if(UAnimSequenceBase* SequenceBase = Cast<UAnimSequenceBase>(Asset))
+				{
+					SequenceBase->Modify();
+
+					if (FAnimCurveBase* CurrentCurveData = SequenceBase->RawCurveData.GetCurveData(CurveToRename.UID))
+					{
+						CurrentCurveData->Name.DisplayName = InNewName;
+						SequenceBase->MarkRawDataAsModified();
+						if (UAnimSequence* Seq = Cast<UAnimSequence>(SequenceBase))
+						{
+							SequencesToRecompress.Add(Seq);
+							Seq->CompressedCurveData.Empty();
+						}
+					}
+				}
+				else if (UPoseAsset* PoseAsset = Cast<UPoseAsset>(Asset))
+				{
+					PoseAsset->Modify();
+
+					PoseAsset->RenameSmartName(CurveToRename.DisplayName, InNewName);
+				}
+			}
+			GWarn->EndSlowTask();
+
+			GWarn->BeginSlowTask(LOCTEXT("RebuildingAnimations", "Rebaking/compressing modified animations"), true);
+
+			//Make sure skeleton is correct before compression 
+
+			Skeleton->RenameSmartnameAndModify(InContainerName, InNameUid, InNewName);
+
+			// Rebake/compress the animations
+			for (UAnimSequence* Seq : SequencesToRecompress)
+			{
+				GWarn->StatusUpdate(1, 2, FText::Format(LOCTEXT("RebuildingAnimationsStatus", "Rebuilding {0}"), FText::FromString(Seq->GetName())));
+				Seq->RequestSyncAnimRecompression();
+			}
+
+			GWarn->EndSlowTask();
+		}
+		else
 		{
-			FScopedTransaction Transaction(LOCTEXT("TransactionRename", "Rename Element"));
 			Skeleton->RenameSmartnameAndModify(InContainerName, InNameUid, InNewName);
 		}
+
+		OnSmartNameChanged.Broadcast(InContainerName);
 	}
-}
 
-void FEditableSkeleton::RemoveSmartname(const FName& InContainerName, SmartName::UID_Type InNameUid)
-{
-	// Note no transaction here as this doesnt do anything at the moment!
-	Skeleton->RemoveSmartnameAndModify(InContainerName, InNameUid);
 }
 
-void FEditableSkeleton::RemoveSmartnamesAndFixupAnimations(const FName& InContainerName, const TArray<SmartName::UID_Type>& InNameUids)
+void FEditableSkeleton::GetAssetsContainingCurves(const FName& InContainerName, const TArray<FName>& InNames, TArray<FAssetData>& OutAssets) const
 {
-	FAssetRegistryModule& AssetModule = FModuleManager::LoadModuleChecked<FAssetRegistryModule>(TEXT("AssetRegistry"));
-	TArray<FAssetData> AnimationAssets;
-	AssetModule.Get().GetAssetsByClass(UAnimSequenceBase::StaticClass()->GetFName(), AnimationAssets, true);
 	FAssetData SkeletonData(Skeleton);
 	const FString CurrentSkeletonName = SkeletonData.GetExportTextName();
 
+	FAssetRegistryModule& AssetModule = FModuleManager::LoadModuleChecked<FAssetRegistryModule>(TEXT("AssetRegistry"));
+	AssetModule.Get().GetAssetsByClass(UAnimationAsset::StaticClass()->GetFName(), OutAssets, true);
+
 	GWarn->BeginSlowTask(LOCTEXT("CollectAnimationsTaskDesc", "Collecting assets..."), true);
 
-	for (int32 Idx = AnimationAssets.Num() - 1; Idx >= 0; --Idx)
+	for (int32 Idx = OutAssets.Num() - 1; Idx >= 0; --Idx)
 	{
-		FAssetData& Data = AnimationAssets[Idx];
-		bool bRemove = true;
+		FAssetData& Data = OutAssets[Idx];
+		bool bAssetContainsRemovableCurves = false;
 
 		const FString SkeletonDataTag = Data.GetTagValueRef<FString>("Skeleton");
 		if (!SkeletonDataTag.IsEmpty() && SkeletonDataTag == CurrentSkeletonName)
 		{
 			FString CurveData;
-			if (!Data.GetTagValue(USkeleton::CurveTag, CurveData))
+			if (!Data.GetTagValue(USkeleton::CurveNameTag, CurveData))
 			{
 				// This asset is old; it hasn't been loaded since before smartnames were added for curves.
 				// unfortunately the only way to delete curves safely is to load old assets to see if they're
@@ -281,51 +352,82 @@ void FEditableSkeleton::RemoveSmartnamesAndFixupAnimations(const FName& InContai
 
 				UObject::FAssetRegistryTag* CurveTag = Tags.FindByPredicate([](const UObject::FAssetRegistryTag& InTag)
 				{
-					return InTag.Name == USkeleton::CurveTag;
+					return InTag.Name == USkeleton::CurveNameTag;
 				});
-				CurveData = CurveTag->Value;
+				
+				if (CurveTag)
+				{
+					CurveData = CurveTag->Value;
+				}
 			}
 
-			TArray<FString> ParsedStringUids;
-			CurveData.ParseIntoArray(ParsedStringUids, *USkeleton::CurveTagDelimiter, true);
-
-			for (const FString& UidString : ParsedStringUids)
+			if(!CurveData.IsEmpty())
 			{
-				SmartName::UID_Type Uid = (SmartName::UID_Type)TCString<TCHAR>::Strtoui64(*UidString, NULL, 10);
-				if (InNameUids.Contains(Uid))
+				TArray<FString> ParsedCurveNames;
+				CurveData.ParseIntoArray(ParsedCurveNames, *USkeleton::CurveTagDelimiter, true);
+
+				for (const FString& CurveString : ParsedCurveNames)
 				{
-					bRemove = false;
-					break;
+					FName CurveName(*CurveString);
+					if (InNames.Contains(CurveName))
+					{
+						bAssetContainsRemovableCurves = true;
+						break;
+					}
 				}
 			}
 		}
 
-		if (bRemove)
+		if (!bAssetContainsRemovableCurves)
 		{
-			AnimationAssets.RemoveAt(Idx);
+			OutAssets.RemoveAtSwap(Idx,1,false);
 		}
 	}
 
 	GWarn->EndSlowTask();
+}
+
+void FEditableSkeleton::RemoveSmartnamesAndFixupAnimations(const FName& InContainerName, const TArray<FName>& InNames)
+{
+	FText Title = LOCTEXT("RemoveCurveInitialDialogTitle", "Confirm Remove");
+	FText ConfirmMessage = LOCTEXT("RemoveCurveInitialDialogMessage", "Removing curves will necessitate loading and modifying animations and pose assets that use these curves. This could be a slow process.\n\nContinue?");
+
+	if (FMessageDialog::Open(EAppMsgType::YesNo, ConfirmMessage, &Title) != EAppReturnType::Yes)
+	{
+		return;
+	}
+
+	TArray<FAssetData> AnimationAssets;
+	GetAssetsContainingCurves(InContainerName, InNames, AnimationAssets);
+
+	bool bRemoved = true;
 
 	// AnimationAssets now only contains assets that are using the selected curve(s)
 	if (AnimationAssets.Num() > 0)
 	{
-		FString ConfirmMessage = LOCTEXT("DeleteCurveMessage", "The following assets use the selected curve(s), deleting will remove the curves from these assets. Continue?\n\n").ToString();
+		bRemoved = false; //need to warn user now
+		FString AssetMessage = LOCTEXT("DeleteCurveMessage", "Deleting curves will:\n\nRemove the curves from Animations and PoseAssets\nRemove poses using that curve name from PoseAssets.\n\nThe following assets will be modified. Continue?\n\n").ToString();
+
+		AnimationAssets.Sort([&](const FAssetData& A, const FAssetData& B)
+		{ 
+			if (A.AssetClass == B.AssetClass)
+			{
+				return A.AssetName < B.AssetName;
+			}
+			return A.AssetClass < B.AssetClass;
+		});
 
 		for (FAssetData& Data : AnimationAssets)
 		{
-			ConfirmMessage += Data.AssetName.ToString() + " (" + Data.AssetClass.ToString() + ")\n";
+			AssetMessage += Data.AssetName.ToString() + " (" + Data.AssetClass.ToString() + ")\n";
 		}
 
-		FText Title = LOCTEXT("DeleteCurveDialogTitle", "Confirm Deletion");
-		FText Message = FText::FromString(ConfirmMessage);
-		if (FMessageDialog::Open(EAppMsgType::YesNo, Message, &Title) == EAppReturnType::No)
-		{
-			return;
-		}
-		else
+		FText AssetTitleText = LOCTEXT("DeleteCurveDialogTitle", "Confirm Deletion");
+		FText AssetMessageText = FText::FromString(AssetMessage);
+
+		if (FMessageDialog::Open(EAppMsgType::YesNo, AssetMessageText, &AssetTitleText) == EAppReturnType::Yes)
 		{
+			bRemoved = true;
 			// Proceed to delete the curves
 			GWarn->BeginSlowTask(FText::Format(LOCTEXT("DeleteCurvesTaskDesc", "Deleting curve from skeleton {0}"), FText::FromString(Skeleton->GetName())), true);
 			FScopedTransaction Transaction(LOCTEXT("DeleteCurvesTransactionName", "Delete skeleton curve"));
@@ -333,17 +435,26 @@ void FEditableSkeleton::RemoveSmartnamesAndFixupAnimations(const FName& InContai
 			// Remove curves from animation assets
 			for (FAssetData& Data : AnimationAssets)
 			{
-				UAnimSequenceBase* Sequence = Cast<UAnimSequenceBase>(Data.GetAsset());
-				USkeleton* MySkeleton = Sequence->GetSkeleton();
-				Sequence->Modify(true);
-				for (SmartName::UID_Type Uid : InNameUids)
+				UObject* Asset = Data.GetAsset();
+
+				if(UAnimSequenceBase* Sequence = Cast<UAnimSequenceBase>(Asset))
 				{
-					FSmartName CurveToDelete;
-					if (MySkeleton->GetSmartNameByUID(USkeleton::AnimCurveMappingName, Uid, CurveToDelete))
+					USkeleton* MySkeleton = Sequence->GetSkeleton();
+					Sequence->Modify(true);
+					for (FName Name : InNames)
 					{
-						Sequence->RawCurveData.DeleteCurveData(CurveToDelete);
-						Sequence->MarkRawDataAsModified();
+						FSmartName CurveToDelete;
+						if (MySkeleton->GetSmartNameByName(USkeleton::AnimCurveMappingName, Name, CurveToDelete))
+						{
+							Sequence->RawCurveData.DeleteCurveData(CurveToDelete);
+						}
 					}
+					Sequence->MarkRawDataAsModified();
+				}
+				else if (UPoseAsset* PoseAsset = Cast<UPoseAsset>(Asset))
+				{
+					PoseAsset->Modify();
+					PoseAsset->RemoveSmartNames(InNames);
 				}
 			}
 			GWarn->EndSlowTask();
@@ -362,13 +473,13 @@ void FEditableSkeleton::RemoveSmartnamesAndFixupAnimations(const FName& InContai
 		}
 	}
 
-	if (InNameUids.Num() > 0)
+	if (bRemoved && InNames.Num() > 0)
 	{
 		// Remove names from skeleton
-		Skeleton->RemoveSmartnamesAndModify(InContainerName, InNameUids);
+		Skeleton->RemoveSmartnamesAndModify(InContainerName, InNames);
 	}
 
-	OnSmartNameRemoved.Broadcast(InContainerName, InNameUids);
+	OnSmartNameChanged.Broadcast(InContainerName);
 }
 
 void FEditableSkeleton::SetCurveMetaDataMaterial(const FSmartName& CurveName, bool bOverrideMaterial)
@@ -382,7 +493,7 @@ void FEditableSkeleton::SetCurveMetaDataMaterial(const FSmartName& CurveName, bo
 	}
 }
 
-void FEditableSkeleton::SetCurveMetaBoneLinks(const FSmartName& CurveName, TArray<FBoneReference>& BoneLinks)
+void FEditableSkeleton::SetCurveMetaBoneLinks(const FSmartName& CurveName, TArray<FBoneReference>& BoneLinks, uint8 InMaxLOD)
 {
 	Skeleton->Modify();
 	FCurveMetaData* CurveMetaData = Skeleton->GetCurveMetaData(CurveName);
@@ -390,7 +501,7 @@ void FEditableSkeleton::SetCurveMetaBoneLinks(const FSmartName& CurveName, TArra
 	{
 		// override curve data
 		CurveMetaData->LinkedBones = BoneLinks;
-		
+		CurveMetaData->MaxLOD = InMaxLOD;
 		//  initialize to this skeleton
 		for (FBoneReference& BoneReference : CurveMetaData->LinkedBones)
 		{
@@ -401,16 +512,20 @@ void FEditableSkeleton::SetCurveMetaBoneLinks(const FSmartName& CurveName, TArra
 
 FName FEditableSkeleton::GenerateUniqueSocketName( FName InName, USkeletalMesh* InSkeletalMesh )
 {
-	int32 NewNumber = InName.GetNumber();
-
-	// Increment NewNumber until we have a unique name (potential infinite loop if *all* int32 values are used!)
-	while ( DoesSocketAlreadyExist( NULL, FText::FromName( FName( InName, NewNumber ) ), ESocketParentType::Skeleton, InSkeletalMesh ) ||
-		(InSkeletalMesh ? DoesSocketAlreadyExist( NULL, FText::FromName( FName( InName, NewNumber ) ), ESocketParentType::Mesh, InSkeletalMesh ) : false ) )
+	if(DoesSocketAlreadyExist(nullptr, FText::FromName(InName), ESocketParentType::Skeleton, InSkeletalMesh ) || DoesSocketAlreadyExist(nullptr, FText::FromName(InName), ESocketParentType::Mesh, InSkeletalMesh))
 	{
-		++NewNumber;
-	}
+		int32 NewNumber = InName.GetNumber();
+
+		// Increment NewNumber until we have a unique name (potential infinite loop if *all* int32 values are used!)
+		while (DoesSocketAlreadyExist(NULL, FText::FromName(FName(InName, NewNumber)), ESocketParentType::Skeleton, InSkeletalMesh) ||
+			(InSkeletalMesh ? DoesSocketAlreadyExist(NULL, FText::FromName(FName(InName, NewNumber)), ESocketParentType::Mesh, InSkeletalMesh) : false))
+		{
+			++NewNumber;
+		}
 
-	return FName( InName, NewNumber );
+		return FName(InName, NewNumber);
+	}
+	return InName;
 }
 
 static USkeletalMeshSocket* FindSocket(const FName& InSocketName, USkeletalMesh* InSkeletalMesh, USkeleton* InSkeleton)
@@ -965,6 +1080,14 @@ int32 FEditableSkeleton::DeleteAnimNotifies(const TArray<FName>& InNotifyNames)
 	return NumAnimationsModified;
 }
 
+void FEditableSkeleton::AddNotify(FName NewName)
+{
+	const FScopedTransaction Transaction(LOCTEXT("AddNewNotifyToSkeleton", "Add New Anim Notify To Skeleton"));
+	Skeleton->Modify();
+	Skeleton->AddNewAnimationNotify(NewName);
+	OnNotifiesChanged.Broadcast();
+}
+
 int32 FEditableSkeleton::RenameNotify(const FName& NewName, const FName& OldName)
 {
 	const FScopedTransaction Transaction(LOCTEXT("RenameAnimNotify", "Rename Anim Notify"));
@@ -1456,14 +1579,14 @@ void FEditableSkeleton::RenameSlotName(const FName& InOldSlotName, const FName&
 	Skeleton->RenameSlotName(InOldSlotName, InNewSlotName);
 }
 
-FDelegateHandle FEditableSkeleton::RegisterOnSmartNameRemoved(const FOnSmartNameRemoved::FDelegate& InOnSmartNameRemoved)
+FDelegateHandle FEditableSkeleton::RegisterOnSmartNameChanged(const FOnSmartNameChanged::FDelegate& InOnSmartNameChanged)
 {
-	return OnSmartNameRemoved.Add(InOnSmartNameRemoved);
+	return OnSmartNameChanged.Add(InOnSmartNameChanged);
 }
 
-void FEditableSkeleton::UnregisterOnSmartNameRemoved(FDelegateHandle InHandle)
+void FEditableSkeleton::UnregisterOnSmartNameChanged(FDelegateHandle InHandle)
 {
-	OnSmartNameRemoved.Remove(InHandle);
+	OnSmartNameChanged.Remove(InHandle);
 }
 
 #undef LOCTEXT_NAMESPACE
diff --git a/Engine/Source/Editor/SkeletonEditor/Private/EditableSkeleton.h b/Engine/Source/Editor/SkeletonEditor/Private/EditableSkeleton.h
--- a/Engine/Source/Editor/SkeletonEditor/Private/EditableSkeleton.h
+++ b/Engine/Source/Editor/SkeletonEditor/Private/EditableSkeleton.h
@@ -37,6 +37,7 @@ public:
 	virtual USkeletalMeshSocket* DuplicateSocket(const FSelectedSocketInfo& SocketInfoToDuplicate, const FName& NewParentBoneName, USkeletalMesh* InSkeletalMesh) override;
 	virtual int32 ValidatePreviewAttachedObjects() override;
 	virtual int32 DeleteAnimNotifies(const TArray<FName>& InSelectedNotifyNames) override;
+	virtual void AddNotify(FName NewName) override;
 	virtual int32 RenameNotify(const FName& NewName, const FName& OldName) override;
 	virtual void GetCompatibleAnimSequences(TArray<struct FAssetData>& OutAssets) override;
 	virtual void RenameSocket(const FName& OldSocketName, const FName& NewSocketName, USkeletalMesh* InSkeletalMesh) override;
@@ -46,10 +47,9 @@ public:
 	virtual void RenameVirtualBone(const FName& OriginalName, const FName& InVBName) override;
 	virtual bool AddSmartname(const FName& InContainerName, const FName& InNewName, FSmartName& OutSmartName) override;
 	virtual void RenameSmartname(const FName& InContainerName, SmartName::UID_Type InNameUid, const FName& InNewName) override;
-	virtual void RemoveSmartname(const FName& InContainerName, SmartName::UID_Type InNameUid) override;
-	virtual void RemoveSmartnamesAndFixupAnimations(const FName& InContainerName, const TArray<SmartName::UID_Type>& InNameUids) override;
+	virtual void RemoveSmartnamesAndFixupAnimations(const FName& InContainerName, const TArray<FName>& InNames) override;
 	virtual void SetCurveMetaDataMaterial(const FSmartName& CurveName, bool bOverrideMaterial) override;
-	virtual void SetCurveMetaBoneLinks(const FSmartName& CurveName, TArray<FBoneReference>& BoneLinks) override;
+	virtual void SetCurveMetaBoneLinks(const FSmartName& CurveName, TArray<FBoneReference>& BoneLinks, uint8 InMaxLOD) override;
 	virtual void SetPreviewMesh(class USkeletalMesh* InSkeletalMesh) override;
 	virtual void LoadAdditionalPreviewSkeletalMeshes() override;
 	virtual void SetAdditionalPreviewSkeletalMeshes(class UDataAsset* InPreviewCollectionAsset) override;
@@ -68,8 +68,8 @@ public:
 	virtual void DeleteSlotName(const FName& InSlotName) override;
 	virtual void DeleteSlotGroup(const FName& InGroupName) override;
 	virtual void RenameSlotName(const FName& InOldSlotName, const FName& InNewSlotName) override;
-	virtual FDelegateHandle RegisterOnSmartNameRemoved(const FOnSmartNameRemoved::FDelegate& InOnSmartNameRemoved) override;
-	virtual void UnregisterOnSmartNameRemoved(FDelegateHandle InHandle) override;
+	virtual FDelegateHandle RegisterOnSmartNameChanged(const FOnSmartNameChanged::FDelegate& InOnSmartNameChanged) override;
+	virtual void UnregisterOnSmartNameChanged(FDelegateHandle InHandle) override;
 	virtual void RegisterOnNotifiesChanged(const FSimpleMulticastDelegate::FDelegate& InDelegate);
 	virtual void UnregisterOnNotifiesChanged(void* Thing);
 	virtual void SetBoneTranslationRetargetingMode(FName InBoneName, EBoneTranslationRetargetingMode::Type NewRetargetingMode) override;
@@ -143,6 +143,9 @@ private:
 	/** Helper function for deleting attached objects */
 	void DeleteAttachedObjects(FPreviewAssetAttachContainer& AttachedAssets, TSharedPtr<class IPersonaPreviewScene> InPreviewScene);
 
+	/** Helper function for finding animations that use certain curves */
+	void GetAssetsContainingCurves(const FName& InContainerName, const TArray<FName>& InNames, TArray<FAssetData>& OutAssets) const;
+
 private:
 	/** The skeleton we are editing */
 	class USkeleton* Skeleton;
@@ -157,7 +160,7 @@ private:
 	FSimpleMulticastDelegate OnTreeRefresh;
 
 	/** Delegate called when a smart name is removed */
-	FOnSmartNameRemoved OnSmartNameRemoved;
+	FOnSmartNameChanged OnSmartNameChanged;
 
 	/** Delegate called when notifies are modified */
 	FSimpleMulticastDelegate OnNotifiesChanged;
