Copying //UE4/Dev-AnimPhys to //UE4/Dev-Main (Source: //UE4/Dev-AnimPhys @ 3624379)

diff --git a/Engine/Source/Editor/SkeletonEditor/Private/SSkeletonTree.cpp b/Engine/Source/Editor/SkeletonEditor/Private/SSkeletonTree.cpp
--- a/Engine/Source/Editor/SkeletonEditor/Private/SSkeletonTree.cpp
+++ b/Engine/Source/Editor/SkeletonEditor/Private/SSkeletonTree.cpp
@@ -59,8 +59,11 @@
 
 #include "BoneSelectionWidget.h"
 #include "BoneProxy.h"
+#include "SkeletonTreeSelection.h"
+#include "SGridPanel.h"
 
 #include "HAL/PlatformApplicationMisc.h"
+#include "STreeView.h"
 
 #define LOCTEXT_NAMESPACE "SSkeletonTree"
 
@@ -68,6 +71,54 @@ const FName	ISkeletonTree::Columns::Name("Name");
 const FName	ISkeletonTree::Columns::Retargeting("Retargeting");
 const FName ISkeletonTree::Columns::BlendProfile("BlendProfile");
 
+// This is mostly duplicated from SListView, to allow for us to avoid selecting collapsed items
+template <typename ItemType>
+class SSkeletonTreeView : public STreeView<ItemType>
+{
+public:
+	bool Private_CanItemBeSelected(ItemType InItem) const
+	{
+		return !(InItem->GetFilterResult() == ESkeletonTreeFilterResult::ShownDescendant && GetMutableDefault<UPersonaOptions>()->bHideParentsWhenFiltering);
+	}
+
+	virtual void Private_SelectRangeFromCurrentTo( ItemType InRangeSelectionEnd ) override
+	{
+		if ( this->SelectionMode.Get() == ESelectionMode::None )
+		{
+			return;
+		}
+
+		const TArray<ItemType>& ItemsSourceRef = (*this->ItemsSource);
+
+		int32 RangeStartIndex = 0;
+		if( TListTypeTraits<ItemType>::IsPtrValid(this->RangeSelectionStart) )
+		{
+			RangeStartIndex = ItemsSourceRef.Find( TListTypeTraits<ItemType>::NullableItemTypeConvertToItemType( this->RangeSelectionStart ) );
+		}
+
+		int32 RangeEndIndex = ItemsSourceRef.Find( InRangeSelectionEnd );
+
+		RangeStartIndex = FMath::Clamp(RangeStartIndex, 0, ItemsSourceRef.Num());
+		RangeEndIndex = FMath::Clamp(RangeEndIndex, 0, ItemsSourceRef.Num());
+
+		if (RangeEndIndex < RangeStartIndex)
+		{
+			Swap( RangeStartIndex, RangeEndIndex );
+		}
+
+		for( int32 ItemIndex = RangeStartIndex; ItemIndex <= RangeEndIndex; ++ItemIndex )
+		{
+			// check if this item can actually be selected
+			if(Private_CanItemBeSelected(ItemsSourceRef[ItemIndex]))
+			{
+				this->SelectedItems.Add( ItemsSourceRef[ItemIndex] );
+			}
+		}
+
+		this->InertialScrollManager.ClearScrollVelocity();
+	}
+};
+
 void SSkeletonTree::Construct(const FArguments& InArgs, const TSharedRef<FEditableSkeleton>& InEditableSkeleton, const FSkeletonTreeArgs& InSkeletonTreeArgs)
 {
 	BoneFilter = EBoneFilter::All;
@@ -80,31 +131,48 @@ void SSkeletonTree::Construct(const FArguments& InArgs, const TSharedRef<FEditab
 	EditableSkeleton = InEditableSkeleton;
 	PreviewScene = InSkeletonTreeArgs.PreviewScene;
 	IsEditable = InArgs._IsEditable;
-	Builder = InArgs._Builder;
+	Mode = InSkeletonTreeArgs.Mode;
+	bAllowMeshOperations = InSkeletonTreeArgs.bAllowMeshOperations;
+	bAllowSkeletonOperations = InSkeletonTreeArgs.bAllowSkeletonOperations;
+	Extenders = InSkeletonTreeArgs.Extenders;
+	OnGetFilterText = InSkeletonTreeArgs.OnGetFilterText;
+	Builder = InSkeletonTreeArgs.Builder;
 	if (!Builder.IsValid())
 	{
-		Builder = MakeShareable(new FSkeletonTreeBuilder(FSkeletonTreeBuilderArgs(), FOnFilterSkeletonTreeItem::CreateSP(this, &SSkeletonTree::HandleFilterSkeletonTreeItem), SharedThis(this), InSkeletonTreeArgs.PreviewScene));
+		Builder = MakeShareable(new FSkeletonTreeBuilder(FSkeletonTreeBuilderArgs()));
 	}
 
+	Builder->Initialize(SharedThis(this), InSkeletonTreeArgs.PreviewScene, FOnFilterSkeletonTreeItem::CreateSP(this, &SSkeletonTree::HandleFilterSkeletonTreeItem));
+
 	TextFilterPtr = MakeShareable(new FTextFilterExpressionEvaluator(ETextFilterExpressionEvaluatorMode::BasicString));
 
 	SetPreviewComponentSocketFilter();
 
 	// Register delegates
-	InSkeletonTreeArgs.OnPostUndo.Add(FSimpleDelegate::CreateSP(this, &SSkeletonTree::PostUndo));
 
 	if(PreviewScene.IsValid())
 	{
 		PreviewScene.Pin()->RegisterOnLODChanged(FSimpleDelegate::CreateSP(this, &SSkeletonTree::OnLODSwitched));
 	}
 
-	InEditableSkeleton->RegisterOnSkeletonHierarchyChanged(USkeleton::FOnSkeletonHierarchyChanged::CreateSP(this, &SSkeletonTree::PostUndo));
+	InEditableSkeleton->RegisterOnSkeletonHierarchyChanged(USkeleton::FOnSkeletonHierarchyChanged::CreateSP(this, &SSkeletonTree::CreateTreeColumns));
+
+	PRAGMA_DISABLE_DEPRECATION_WARNINGS
+	if(InSkeletonTreeArgs.OnObjectSelected.IsBound())
+	{
+		RegisterOnObjectSelected(InSkeletonTreeArgs.OnObjectSelected);
+	}
+	PRAGMA_ENABLE_DEPRECATION_WARNINGS
 
-	RegisterOnObjectSelected(InSkeletonTreeArgs.OnObjectSelected);
+	if (InSkeletonTreeArgs.OnSelectionChanged.IsBound())
+	{
+		RegisterOnSelectionChanged(InSkeletonTreeArgs.OnSelectionChanged);
+	}
 
 	BoneProxy = NewObject<UBoneProxy>(GetTransientPackage());
 	BoneProxy->AddToRoot();
 	BoneProxy->SkelMeshComponent = PreviewScene.IsValid() ? PreviewScene.Pin()->GetPreviewMeshComponent() : nullptr;
+	BoneProxy->bIsTickable = true;
 
 	// Register and bind all our menu commands
 	FSkeletonTreeCommands::Register();
@@ -112,77 +180,87 @@ void SSkeletonTree::Construct(const FArguments& InArgs, const TSharedRef<FEditab
 
 	this->ChildSlot
 	[
-		SNew( SVerticalBox )
-		
-		+ SVerticalBox::Slot()
-		.AutoHeight()
-		.Padding( FMargin( 0.0f, 0.0f, 0.0f, 4.0f ) )
+		SNew( SOverlay )
+		+SOverlay::Slot()
 		[
-			SAssignNew( NameFilterBox, SSearchBox )
-			.SelectAllTextWhenFocused( true )
-			.OnTextChanged( this, &SSkeletonTree::OnFilterTextChanged )
-			.HintText( LOCTEXT( "SearchBoxHint", "Search Skeleton Tree...") )
-			.AddMetaData<FTagMetaData>(TEXT("SkelTree.Search"))
+			// Add a border if we are being used as a picker
+			SNew(SBorder)
+			.Visibility_Lambda([this](){ return Mode == ESkeletonTreeMode::Picker ? EVisibility::Visible: EVisibility::Collapsed; })
+			.BorderImage(FEditorStyle::Get().GetBrush("Menu.Background"))
 		]
-
-		+ SVerticalBox::Slot()
-		.Padding( FMargin( 0.0f, 4.0f, 0.0f, 0.0f ) )
-		[
-			SAssignNew(TreeHolder, SOverlay)
-		]
-
-		+ SVerticalBox::Slot()
-		.AutoHeight()
+		+SOverlay::Slot()
 		[
-			SNew( SHorizontalBox )
+			SNew( SVerticalBox )
+		
+			+ SVerticalBox::Slot()
+			.AutoHeight()
+			.Padding( FMargin( 0.0f, 0.0f, 0.0f, 4.0f ) )
+			[
+				SAssignNew( NameFilterBox, SSearchBox )
+				.SelectAllTextWhenFocused( true )
+				.OnTextChanged( this, &SSkeletonTree::OnFilterTextChanged )
+				.HintText( LOCTEXT( "SearchBoxHint", "Search Skeleton Tree...") )
+				.AddMetaData<FTagMetaData>(TEXT("SkelTree.Search"))
+			]
 
-			+ SHorizontalBox::Slot()
-			.Padding(0.0f, 0.0f, 2.0f, 0.0f)
-			.VAlign(VAlign_Center)
-			.HAlign(HAlign_Left)
-			.FillWidth(1.0f)
+			+ SVerticalBox::Slot()
+			.Padding( FMargin( 0.0f, 4.0f, 0.0f, 0.0f ) )
 			[
-				SNew(SHorizontalBox)
-				
-				+ SHorizontalBox::Slot()
-				.Padding(0.0f, 0.0f, 2.0f, 0.0f)
-				.VAlign(VAlign_Center)
-				[
-					SAssignNew(BlendProfilePicker, SBlendProfilePicker, GetEditableSkeleton())
-					.Standalone(true)
-					.OnBlendProfileSelected(this, &SSkeletonTree::OnBlendProfileSelected)
-				]
+				SAssignNew(TreeHolder, SOverlay)
 			]
 
-			+ SHorizontalBox::Slot()
-			.Padding( 0.0f, 0.0f, 2.0f, 0.0f )
-			.VAlign(VAlign_Center)
-			.HAlign(HAlign_Right)
-			.AutoWidth()
+			+ SVerticalBox::Slot()
+			.AutoHeight()
 			[
-				SAssignNew( FilterComboButton, SComboButton )
-				.ButtonStyle(FEditorStyle::Get(), "ToggleButton")
-				.ForegroundColor(this, &SSkeletonTree::GetFilterComboButtonForegroundColor)
-				.ContentPadding(0.0f)
-				.OnGetMenuContent( this, &SSkeletonTree::CreateFilterMenu )
-				.ToolTipText( LOCTEXT( "BoneFilterToolTip", "Change which types of bones and sockets are shown, retargeting options" ) )
-				.AddMetaData<FTagMetaData>(TEXT("SkelTree.Bones"))
-				.ButtonContent()
+				SNew( SHorizontalBox )
+
+				+ SHorizontalBox::Slot()
+				.Padding(0.0f, 0.0f, InSkeletonTreeArgs.bShowBlendProfiles ? 2.0f : 0.0f, 0.0f)
+				.VAlign(VAlign_Center)
+				.HAlign(HAlign_Left)
 				[
 					SNew(SHorizontalBox)
-					+SHorizontalBox::Slot()
-					.AutoWidth()
+					.Visibility(InSkeletonTreeArgs.bShowBlendProfiles ? EVisibility::Visible : EVisibility::Collapsed)
+					+ SHorizontalBox::Slot()
+					.Padding(0.0f, 0.0f, 2.0f, 0.0f)
 					.VAlign(VAlign_Center)
 					[
-						SNew(SImage).Image( FEditorStyle::GetBrush("GenericViewButton") )
+						SAssignNew(BlendProfilePicker, SBlendProfilePicker, GetEditableSkeleton())
+						.Standalone(true)
+						.OnBlendProfileSelected(this, &SSkeletonTree::OnBlendProfileSelected)
 					]
-					+SHorizontalBox::Slot()
-					.AutoWidth()
-					.Padding(2, 0, 0, 0)
-					.VAlign(VAlign_Center)
+				]
+
+				+ SHorizontalBox::Slot()
+				.Padding( 0.0f, 0.0f, InSkeletonTreeArgs.bShowFilterMenu ? 2.0f : 0.0f, 0.0f )
+				.VAlign(VAlign_Center)
+				.HAlign(HAlign_Right)
+				[
+					SAssignNew( FilterComboButton, SComboButton )
+					.Visibility(InSkeletonTreeArgs.bShowFilterMenu ? EVisibility::Visible : EVisibility::Collapsed)
+					.ButtonStyle(FEditorStyle::Get(), "ToggleButton")
+					.ForegroundColor(this, &SSkeletonTree::GetFilterComboButtonForegroundColor)
+					.ContentPadding(0.0f)
+					.OnGetMenuContent( this, &SSkeletonTree::CreateFilterMenu )
+					.ToolTipText( LOCTEXT( "BoneFilterToolTip", "Change which types of tree items are shown, as well as other options" ) )
+					.AddMetaData<FTagMetaData>(TEXT("SkelTree.Bones"))
+					.ButtonContent()
 					[
-						SNew( STextBlock )
-						.Text( this, &SSkeletonTree::GetFilterMenuTitle )
+						SNew(SHorizontalBox)
+						+SHorizontalBox::Slot()
+						.AutoWidth()
+						.VAlign(VAlign_Center)
+						[
+							SNew(SImage).Image( FEditorStyle::GetBrush("GenericViewButton") )
+						]
+						+SHorizontalBox::Slot()
+						.AutoWidth()
+						.Padding(2, 0, 0, 0)
+						.VAlign(VAlign_Center)
+						[
+							SNew( STextBlock )
+							.Text( this, &SSkeletonTree::GetFilterMenuTitle )
+						]
 					]
 				]
 			]
@@ -191,9 +269,12 @@ void SSkeletonTree::Construct(const FArguments& InArgs, const TSharedRef<FEditab
 
 	CreateTreeColumns();
 
+	SetInitialExpansionState();
+
 	OnLODSwitched();
 }
 
+PRAGMA_DISABLE_DEPRECATION_WARNINGS
 SSkeletonTree::~SSkeletonTree()
 {
 	if (EditableSkeleton.IsValid())
@@ -207,6 +288,7 @@ SSkeletonTree::~SSkeletonTree()
 		BoneProxy = nullptr;
 	}
 }
+PRAGMA_ENABLE_DEPRECATION_WARNINGS
 
 void SSkeletonTree::BindCommands()
 {
@@ -296,6 +378,12 @@ void SSkeletonTree::BindCommands()
 		FCanExecuteAction(),
 		FIsActionChecked::CreateLambda([]() { return GetDefault<UPersonaOptions>()->bFlattenSkeletonHierarchyWhenFiltering; }));
 
+	CommandList.MapAction(
+		MenuActions.HideParentsWhenFiltering,
+		FExecuteAction::CreateLambda([this]() { GetMutableDefault<UPersonaOptions>()->bHideParentsWhenFiltering = !GetDefault<UPersonaOptions>()->bHideParentsWhenFiltering; }),
+		FCanExecuteAction(),
+		FIsActionChecked::CreateLambda([]() { return GetDefault<UPersonaOptions>()->bHideParentsWhenFiltering; }));
+
 	// Socket manipulation commands
 	CommandList.MapAction(
 		MenuActions.AddSocket,
@@ -313,7 +401,8 @@ void SSkeletonTree::BindCommands()
 
 	CommandList.MapAction(
 		MenuActions.RemoveMeshSocket,
-		FExecuteAction::CreateSP( this, &SSkeletonTree::OnDeleteSelectedRows ) ); // Removing customization just deletes the mesh socket
+		FExecuteAction::CreateSP( this, &SSkeletonTree::OnDeleteSelectedRows ),
+		FCanExecuteAction::CreateSP( this, &SSkeletonTree::CanDeleteSelectedRows ) );
 
 	CommandList.MapAction(
 		MenuActions.PromoteSocketToSkeleton,
@@ -321,7 +410,8 @@ void SSkeletonTree::BindCommands()
 
 	CommandList.MapAction(
 		MenuActions.DeleteSelectedRows,
-		FExecuteAction::CreateSP( this, &SSkeletonTree::OnDeleteSelectedRows ) );
+		FExecuteAction::CreateSP( this, &SSkeletonTree::OnDeleteSelectedRows ),
+		FCanExecuteAction::CreateSP( this, &SSkeletonTree::CanDeleteSelectedRows ));
 
 	CommandList.MapAction(
 		MenuActions.CopyBoneNames,
@@ -373,7 +463,7 @@ struct FScopedSavedSelection
 			TArray<TSharedPtr<ISkeletonTreeItem>> SelectedItems = InSkeletonTree->SkeletonTreeView->GetSelectedItems();
 			for (const TSharedPtr<ISkeletonTreeItem>& SelectedItem : SelectedItems)
 			{
-				SavedSelections.Add({ SelectedItem->GetRowItemName(), SelectedItem->GetTypeName() });
+				SavedSelections.Add({ SelectedItem->GetRowItemName(), SelectedItem->GetTypeName(), SelectedItem->GetObject() });
 			}
 		}
 	}
@@ -389,7 +479,7 @@ struct FScopedSavedSelection
 				{
 					for (FSavedSelection& SavedSelection : SavedSelections)
 					{
-						if (Item->GetRowItemName() == SavedSelection.ItemName && Item->GetTypeName() == SavedSelection.ItemType)
+						if (Item->GetRowItemName() == SavedSelection.ItemName && Item->GetTypeName() == SavedSelection.ItemType && Item->GetObject() == SavedSelection.ItemObject)
 						{
 							SkeletonTree->SkeletonTreeView->SetItemSelection(Item, true);
 							break;
@@ -407,6 +497,9 @@ struct FScopedSavedSelection
 
 		/** Type of the selected item */
 		FName ItemType;
+
+		/** Object of selected item */
+		UObject* ItemObject;
 	};
 
 	TSharedPtr<SSkeletonTree> SkeletonTree;
@@ -416,7 +509,19 @@ struct FScopedSavedSelection
 
 void SSkeletonTree::CreateTreeColumns()
 {
-	TSharedRef<SHeaderRow> TreeHeaderRow = SNew(SHeaderRow)
+	auto HeaderVisibilityLambda = [this]()
+	{
+		if(Mode == ESkeletonTreeMode::Editor)
+		{
+			return bShowingAdvancedOptions || BlendProfilePicker->GetSelectedBlendProfileName() != NAME_None ? EVisibility::Visible : EVisibility::Collapsed;
+		}
+		
+		return  EVisibility::Collapsed; 
+	};
+		
+	TSharedRef<SHeaderRow> TreeHeaderRow = 
+		SNew(SHeaderRow)
+		.Visibility_Lambda(HeaderVisibilityLambda)
 		+ SHeaderRow::Column(ISkeletonTree::Columns::Name)
 		.DefaultLabel(LOCTEXT("SkeletonBoneNameLabel", "Name"))
 		.FillWidth(0.5f);
@@ -441,7 +546,7 @@ void SSkeletonTree::CreateTreeColumns()
 	{
 		FScopedSavedSelection ScopedSelection(SharedThis(this));
 
-		SkeletonTreeView = SNew(STreeView<TSharedPtr<ISkeletonTreeItem>>)
+		SkeletonTreeView = SNew(SSkeletonTreeView<TSharedPtr<ISkeletonTreeItem>>)
 		.TreeItemsSource(&FilteredItems)
 		.OnGenerateRow(this, &SSkeletonTree::MakeTreeRowWidget)
 		.OnGetChildren(this, &SSkeletonTree::GetFilteredChildren)
@@ -489,103 +594,45 @@ void SSkeletonTree::ApplyFilter()
 
 	FilteredItems.Empty();
 
-	FSkeletonTreeFilterArgs FilterArgs;
-	FilterArgs.bWillFilter = (!FilterText.IsEmpty());
+	FSkeletonTreeFilterArgs FilterArgs(!FilterText.IsEmpty() ? TextFilterPtr : nullptr);
 	FilterArgs.bFlattenHierarchyOnFilter = GetDefault<UPersonaOptions>()->bFlattenSkeletonHierarchyWhenFiltering;
 	Builder->Filter(FilterArgs, Items, FilteredItems);
 
-	for (TSharedPtr<ISkeletonTreeItem>& Item : LinearItems)
+	if(!FilterText.IsEmpty())
 	{
-		SkeletonTreeView->SetItemExpansion(Item, true);
-	}
-
-	SkeletonTreeView->RequestTreeRefresh();
-}
-
-class FBoneTreeSelection
-{
-public:
-	TArray<TSharedPtr<ISkeletonTreeItem>> SelectedItems;
-
-	TArray<TSharedPtr<FSkeletonTreeBoneItem>> SelectedBones;
-	TArray<TSharedPtr<FSkeletonTreeSocketItem>> SelectedSockets;
-	TArray<TSharedPtr<FSkeletonTreeAttachedAssetItem>> SelectedAssets;
-	TArray<TSharedPtr<FSkeletonTreeVirtualBoneItem>> SelectedVirtualBones;
-
-	FBoneTreeSelection(TArray<TSharedPtr<ISkeletonTreeItem>> InSelectedItems) : SelectedItems(InSelectedItems)
-	{
-		for ( auto ItemIt = SelectedItems.CreateConstIterator(); ItemIt; ++ItemIt )
+		for (TSharedPtr<ISkeletonTreeItem>& Item : LinearItems)
 		{
-			TSharedPtr<ISkeletonTreeItem> Item = *(ItemIt);
-
-			if (Item->IsOfType<FSkeletonTreeBoneItem>())
-			{
-				SelectedBones.Add(StaticCastSharedPtr<FSkeletonTreeBoneItem>(Item));
-			}
-			else if (Item->IsOfType<FSkeletonTreeSocketItem>())
-			{
-				SelectedSockets.Add(StaticCastSharedPtr<FSkeletonTreeSocketItem>(Item));
-			}
-			else if (Item->IsOfType<FSkeletonTreeAttachedAssetItem>())
+			if (Item->GetFilterResult() > ESkeletonTreeFilterResult::Hidden)
 			{
-				SelectedAssets.Add(StaticCastSharedPtr<FSkeletonTreeAttachedAssetItem>(Item));
-			}
-			else if (Item->IsOfType<FSkeletonTreeVirtualBoneItem>())
-			{
-				SelectedVirtualBones.Add(StaticCastSharedPtr<FSkeletonTreeVirtualBoneItem>(Item));
+				SkeletonTreeView->SetItemExpansion(Item, true);
 			}
 		}
 	}
-
-	bool IsMultipleItemsSelected() const
-	{
-		return SelectedItems.Num() > 1;
-	}
-
-	bool IsSingleItemSelected() const
+	else
 	{
-		return SelectedItems.Num() == 1;
+		SetInitialExpansionState();
 	}
 
-	template<typename ItemType>
-	bool IsSingleOfTypeSelected() const
-	{
-		if(IsSingleItemSelected())
-		{
-			return SelectedItems[0]->IsOfType<ItemType>();
-		}
-		return false;
-	}
-
-	TSharedPtr<ISkeletonTreeItem> GetSingleSelectedItem()
-	{
-		check(IsSingleItemSelected());
-		return SelectedItems[0];
-	}
+	SkeletonTreeView->RequestTreeRefresh();
+}
 
-	template<typename ItemType>
-	bool HasSelectedOfType() const
+void SSkeletonTree::SetInitialExpansionState()
+{
+	for (TSharedPtr<ISkeletonTreeItem>& Item : LinearItems)
 	{
-		for (const TSharedPtr<ISkeletonTreeItem>& SelectedItem : SelectedItems)
-		{
-			if (SelectedItem->IsOfType<ItemType>())
-			{
-				return true;
-			}
-		}
-
-		return false;
+		SkeletonTreeView->SetItemExpansion(Item, Item->IsInitiallyExpanded());
 	}
-};
+}
 
 TSharedPtr< SWidget > SSkeletonTree::CreateContextMenu()
 {
 	const FSkeletonTreeCommands& Actions = FSkeletonTreeCommands::Get();
 
-	FBoneTreeSelection BoneTreeSelection(SkeletonTreeView->GetSelectedItems());
+	TArray<TSharedPtr<ISkeletonTreeItem>> SelectedItems = SkeletonTreeView->GetSelectedItems();
+	FSkeletonTreeSelection BoneTreeSelection(SelectedItems);
 
 	const bool CloseAfterSelection = true;
-	FMenuBuilder MenuBuilder( CloseAfterSelection, UICommandList );
+	FMenuBuilder MenuBuilder( CloseAfterSelection, UICommandList, Extenders );
 	{
 		if(BoneTreeSelection.HasSelectedOfType<FSkeletonTreeAttachedAssetItem>() || BoneTreeSelection.HasSelectedOfType<FSkeletonTreeSocketItem>() || BoneTreeSelection.HasSelectedOfType<FSkeletonTreeVirtualBoneItem>())
 		{
@@ -594,6 +641,7 @@ TSharedPtr< SWidget > SSkeletonTree::CreateContextMenu()
 			MenuBuilder.EndSection();
 		}
 
+
 		const bool bNeedsBoneActionsHeading = BoneTreeSelection.HasSelectedOfType<FSkeletonTreeBoneItem>() || BoneTreeSelection.HasSelectedOfType<FSkeletonTreeVirtualBoneItem>();
 
 		if (bNeedsBoneActionsHeading)
@@ -606,7 +654,7 @@ TSharedPtr< SWidget > SSkeletonTree::CreateContextMenu()
 			MenuBuilder.AddMenuEntry(Actions.CopyBoneNames);
 			MenuBuilder.AddMenuEntry(Actions.ResetBoneTransforms);
 
-			if (BoneTreeSelection.IsSingleOfTypeSelected<FSkeletonTreeBoneItem>())
+			if (BoneTreeSelection.IsSingleOfTypeSelected<FSkeletonTreeBoneItem>() && bAllowSkeletonOperations)
 			{
 				MenuBuilder.AddMenuEntry(Actions.AddSocket);
 				MenuBuilder.AddMenuEntry(Actions.PasteSockets);
@@ -616,82 +664,91 @@ TSharedPtr< SWidget > SSkeletonTree::CreateContextMenu()
 
 		if (bNeedsBoneActionsHeading)
 		{
-			MenuBuilder.AddSubMenu(LOCTEXT("AddVirtualBone", "Add Virtual Bone"),
-				LOCTEXT("AddVirtualBone_ToolTip", "Adds a virtual bone to the skeleton."),
-				FNewMenuDelegate::CreateSP(this, &SSkeletonTree::FillVirtualBoneSubmenu, BoneTreeSelection.SelectedItems));
+			if (bAllowSkeletonOperations)
+			{
+				MenuBuilder.AddSubMenu(LOCTEXT("AddVirtualBone", "Add Virtual Bone"),
+					LOCTEXT("AddVirtualBone_ToolTip", "Adds a virtual bone to the skeleton."),
+					FNewMenuDelegate::CreateSP(this, &SSkeletonTree::FillVirtualBoneSubmenu, SelectedItems));
+			}
 
 			MenuBuilder.EndSection();
 		}
 
-		if(BoneTreeSelection.HasSelectedOfType<FSkeletonTreeBoneItem>())
+		if (bAllowSkeletonOperations)
 		{
-			UBlendProfile* const SelectedBlendProfile = BlendProfilePicker->GetSelectedBlendProfile();
-			if(SelectedBlendProfile && BoneTreeSelection.IsSingleOfTypeSelected<FSkeletonTreeBoneItem>())
+			if(BoneTreeSelection.HasSelectedOfType<FSkeletonTreeBoneItem>())
 			{
-				TSharedPtr<FSkeletonTreeBoneItem> BoneItem = BoneTreeSelection.SelectedBones[0];
+				UBlendProfile* const SelectedBlendProfile = BlendProfilePicker->GetSelectedBlendProfile();
+				if(SelectedBlendProfile && BoneTreeSelection.IsSingleOfTypeSelected<FSkeletonTreeBoneItem>())
+				{
+					TSharedPtr<FSkeletonTreeBoneItem> BoneItem = BoneTreeSelection.GetSelectedItems<FSkeletonTreeBoneItem>()[0];
 
-				FName BoneName = BoneItem->GetAttachName();
-				const USkeleton& Skeleton = GetEditableSkeletonInternal()->GetSkeleton();
-				int32 BoneIndex = Skeleton.GetReferenceSkeleton().FindBoneIndex(BoneName);
+					FName BoneName = BoneItem->GetAttachName();
+					const USkeleton& Skeleton = GetEditableSkeletonInternal()->GetSkeleton();
+					int32 BoneIndex = Skeleton.GetReferenceSkeleton().FindBoneIndex(BoneName);
 
-				float CurrentBlendScale = SelectedBlendProfile->GetBoneBlendScale(BoneIndex);
+					float CurrentBlendScale = SelectedBlendProfile->GetBoneBlendScale(BoneIndex);
 
-				MenuBuilder.BeginSection("SkeletonTreeBlendProfileScales", LOCTEXT("BlendProfileContextOptions", "Blend Profile"));
-				{
-					FUIAction RecursiveSetScales;
-					RecursiveSetScales.ExecuteAction = FExecuteAction::CreateSP(this, &SSkeletonTree::RecursiveSetBlendProfileScales, CurrentBlendScale);
+					MenuBuilder.BeginSection("SkeletonTreeBlendProfileScales", LOCTEXT("BlendProfileContextOptions", "Blend Profile"));
+					{
+						FUIAction RecursiveSetScales;
+						RecursiveSetScales.ExecuteAction = FExecuteAction::CreateSP(this, &SSkeletonTree::RecursiveSetBlendProfileScales, CurrentBlendScale);
 					
-					MenuBuilder.AddMenuEntry
-						(
-						FText::Format(LOCTEXT("RecursiveSetBlendScales_Label", "Recursively Set Blend Scales To {0}"), FText::AsNumber(CurrentBlendScale)),
-						LOCTEXT("RecursiveSetBlendScales_ToolTip", "Sets all child bones to use the same blend profile scale as the selected bone"),
-						FSlateIcon(),
-						RecursiveSetScales
-						);
+						MenuBuilder.AddMenuEntry
+							(
+							FText::Format(LOCTEXT("RecursiveSetBlendScales_Label", "Recursively Set Blend Scales To {0}"), FText::AsNumber(CurrentBlendScale)),
+							LOCTEXT("RecursiveSetBlendScales_ToolTip", "Sets all child bones to use the same blend profile scale as the selected bone"),
+							FSlateIcon(),
+							RecursiveSetScales
+							);
+					}
+					MenuBuilder.EndSection();
 				}
-				MenuBuilder.EndSection();
-			}
 
-			if(bShowingAdvancedOptions)
-			{
-				MenuBuilder.BeginSection("SkeletonTreeBoneTranslationRetargeting", LOCTEXT("BoneTranslationRetargetingHeader", "Bone Translation Retargeting"));
+				if(bShowingAdvancedOptions)
 				{
-					FUIAction RecursiveRetargetingSkeletonAction = FUIAction(FExecuteAction::CreateSP(this, &SSkeletonTree::SetBoneTranslationRetargetingModeRecursive, EBoneTranslationRetargetingMode::Skeleton));
-					FUIAction RecursiveRetargetingAnimationAction = FUIAction(FExecuteAction::CreateSP(this, &SSkeletonTree::SetBoneTranslationRetargetingModeRecursive, EBoneTranslationRetargetingMode::Animation));
-					FUIAction RecursiveRetargetingAnimationScaledAction = FUIAction(FExecuteAction::CreateSP(this, &SSkeletonTree::SetBoneTranslationRetargetingModeRecursive, EBoneTranslationRetargetingMode::AnimationScaled));
-					FUIAction RecursiveRetargetingAnimationRelativeAction = FUIAction(FExecuteAction::CreateSP(this, &SSkeletonTree::SetBoneTranslationRetargetingModeRecursive, EBoneTranslationRetargetingMode::AnimationRelative));
-
-					MenuBuilder.AddMenuEntry
-						(LOCTEXT("SetTranslationRetargetingSkeletonChildrenAction", "Recursively Set Translation Retargeting Skeleton")
-						, LOCTEXT("BoneTranslationRetargetingSkeletonToolTip", "Use translation from Skeleton.")
-						, FSlateIcon()
-						, RecursiveRetargetingSkeletonAction
-						);
-
-					MenuBuilder.AddMenuEntry
-						(LOCTEXT("SetTranslationRetargetingAnimationChildrenAction", "Recursively Set Translation Retargeting Animation")
-						, LOCTEXT("BoneTranslationRetargetingAnimationToolTip", "Use translation from animation.")
-						, FSlateIcon()
-						, RecursiveRetargetingAnimationAction
-						);
-
-					MenuBuilder.AddMenuEntry
-						(LOCTEXT("SetTranslationRetargetingAnimationScaledChildrenAction", "Recursively Set Translation Retargeting AnimationScaled")
-						, LOCTEXT("BoneTranslationRetargetingAnimationScaledToolTip", "Use translation from animation, scale length by Skeleton's proportions.")
-						, FSlateIcon()
-						, RecursiveRetargetingAnimationScaledAction
-						);
-
-					MenuBuilder.AddMenuEntry
-						(LOCTEXT("SetTranslationRetargetingAnimationRelativeChildrenAction", "Recursively Set Translation Retargeting AnimationRelative")
-						, LOCTEXT("BoneTranslationRetargetingAnimationRelativeToolTip", "Use relative translation from animation similar to an additive animation.")
-						, FSlateIcon()
-						, RecursiveRetargetingAnimationRelativeAction
-						);
+					MenuBuilder.BeginSection("SkeletonTreeBoneTranslationRetargeting", LOCTEXT("BoneTranslationRetargetingHeader", "Bone Translation Retargeting"));
+					{
+						FUIAction RecursiveRetargetingSkeletonAction = FUIAction(FExecuteAction::CreateSP(this, &SSkeletonTree::SetBoneTranslationRetargetingModeRecursive, EBoneTranslationRetargetingMode::Skeleton));
+						FUIAction RecursiveRetargetingAnimationAction = FUIAction(FExecuteAction::CreateSP(this, &SSkeletonTree::SetBoneTranslationRetargetingModeRecursive, EBoneTranslationRetargetingMode::Animation));
+						FUIAction RecursiveRetargetingAnimationScaledAction = FUIAction(FExecuteAction::CreateSP(this, &SSkeletonTree::SetBoneTranslationRetargetingModeRecursive, EBoneTranslationRetargetingMode::AnimationScaled));
+						FUIAction RecursiveRetargetingAnimationRelativeAction = FUIAction(FExecuteAction::CreateSP(this, &SSkeletonTree::SetBoneTranslationRetargetingModeRecursive, EBoneTranslationRetargetingMode::AnimationRelative));
+
+						MenuBuilder.AddMenuEntry
+							(LOCTEXT("SetTranslationRetargetingSkeletonChildrenAction", "Recursively Set Translation Retargeting Skeleton")
+							, LOCTEXT("BoneTranslationRetargetingSkeletonToolTip", "Use translation from Skeleton.")
+							, FSlateIcon()
+							, RecursiveRetargetingSkeletonAction
+							);
+
+						MenuBuilder.AddMenuEntry
+							(LOCTEXT("SetTranslationRetargetingAnimationChildrenAction", "Recursively Set Translation Retargeting Animation")
+							, LOCTEXT("BoneTranslationRetargetingAnimationToolTip", "Use translation from animation.")
+							, FSlateIcon()
+							, RecursiveRetargetingAnimationAction
+							);
+
+						MenuBuilder.AddMenuEntry
+							(LOCTEXT("SetTranslationRetargetingAnimationScaledChildrenAction", "Recursively Set Translation Retargeting AnimationScaled")
+							, LOCTEXT("BoneTranslationRetargetingAnimationScaledToolTip", "Use translation from animation, scale length by Skeleton's proportions.")
+							, FSlateIcon()
+							, RecursiveRetargetingAnimationScaledAction
+							);
+
+						MenuBuilder.AddMenuEntry
+							(LOCTEXT("SetTranslationRetargetingAnimationRelativeChildrenAction", "Recursively Set Translation Retargeting AnimationRelative")
+							, LOCTEXT("BoneTranslationRetargetingAnimationRelativeToolTip", "Use relative translation from animation similar to an additive animation.")
+							, FSlateIcon()
+							, RecursiveRetargetingAnimationRelativeAction
+							);
+					}
+					MenuBuilder.EndSection();
 				}
-				MenuBuilder.EndSection();
 			}
+		}
 
+		if(bAllowMeshOperations)
+		{
 			MenuBuilder.BeginSection("SkeletonTreeBoneReductionForLOD", LOCTEXT("BoneReductionHeader", "LOD Bone Reduction"));
 			{
 				MenuBuilder.AddSubMenu(
@@ -709,70 +766,80 @@ TSharedPtr< SWidget > SSkeletonTree::CreateContextMenu()
 			MenuBuilder.EndSection();
 		}
 
-		if (BoneTreeSelection.HasSelectedOfType<FSkeletonTreeVirtualBoneItem>())
+		if(bAllowSkeletonOperations)
 		{
-			MenuBuilder.BeginSection("SkeletonTreeVirtualBoneActions", LOCTEXT("VirtualBoneActions", "Selected Virtual Bone Actions"));
-
-			if (BoneTreeSelection.IsSingleOfTypeSelected<FSkeletonTreeVirtualBoneItem>())
+			if (BoneTreeSelection.HasSelectedOfType<FSkeletonTreeVirtualBoneItem>())
 			{
-				MenuBuilder.AddMenuEntry(FGenericCommands::Get().Rename, NAME_None, LOCTEXT("RenameVirtualBone_Label", "Rename Virtual Bone"), LOCTEXT("RenameVirtualBone_Tooltip", "Rename this virtual bone"));
-			}
+				MenuBuilder.BeginSection("SkeletonTreeVirtualBoneActions", LOCTEXT("VirtualBoneActions", "Selected Virtual Bone Actions"));
 
-			MenuBuilder.EndSection();
-		}
-
-		if(BoneTreeSelection.HasSelectedOfType<FSkeletonTreeSocketItem>())
-		{
-			MenuBuilder.BeginSection("SkeletonTreeSocketsActions", LOCTEXT( "SocketActions", "Selected Socket Actions" ) );
+				if (BoneTreeSelection.IsSingleOfTypeSelected<FSkeletonTreeVirtualBoneItem>())
+				{
+					MenuBuilder.AddMenuEntry(FGenericCommands::Get().Rename, NAME_None, LOCTEXT("RenameVirtualBone_Label", "Rename Virtual Bone"), LOCTEXT("RenameVirtualBone_Tooltip", "Rename this virtual bone"));
+				}
 
-			MenuBuilder.AddMenuEntry( Actions.CopySockets );
+				MenuBuilder.EndSection();
+			}
 
-			if(BoneTreeSelection.IsSingleOfTypeSelected<FSkeletonTreeSocketItem>())
+			if(BoneTreeSelection.HasSelectedOfType<FSkeletonTreeSocketItem>())
 			{
-				MenuBuilder.AddMenuEntry( FGenericCommands::Get().Rename, NAME_None, LOCTEXT("RenameSocket_Label", "Rename Socket"), LOCTEXT("RenameSocket_Tooltip", "Rename this socket") );
+				MenuBuilder.BeginSection("SkeletonTreeSocketsActions", LOCTEXT( "SocketActions", "Selected Socket Actions" ) );
 
-				TSharedPtr<FSkeletonTreeSocketItem> SocketItem = BoneTreeSelection.SelectedSockets[0];
+				MenuBuilder.AddMenuEntry( Actions.CopySockets );
 
-				if (SocketItem->IsSocketCustomized() && SocketItem->GetParentType() == ESocketParentType::Mesh )
+				if(BoneTreeSelection.IsSingleOfTypeSelected<FSkeletonTreeSocketItem>())
 				{
-					MenuBuilder.AddMenuEntry( Actions.RemoveMeshSocket );
-				}
+					MenuBuilder.AddMenuEntry( FGenericCommands::Get().Rename, NAME_None, LOCTEXT("RenameSocket_Label", "Rename Socket"), LOCTEXT("RenameSocket_Tooltip", "Rename this socket") );
 
-				// If the socket is on the skeleton, we have a valid mesh
-				// and there isn't one of the same name on the mesh, we can customize it
-				if (SocketItem->CanCustomizeSocket() )
-				{
-					if (SocketItem->GetParentType() == ESocketParentType::Skeleton )
+					TSharedPtr<FSkeletonTreeSocketItem> SocketItem = BoneTreeSelection.GetSelectedItems<FSkeletonTreeSocketItem>()[0];
+
+					if (SocketItem->IsSocketCustomized() && SocketItem->GetParentType() == ESocketParentType::Mesh )
 					{
-						MenuBuilder.AddMenuEntry( Actions.CreateMeshSocket );
+						MenuBuilder.AddMenuEntry( Actions.RemoveMeshSocket );
 					}
-					else if (SocketItem->GetParentType() == ESocketParentType::Mesh )
+
+					// If the socket is on the skeleton, we have a valid mesh
+					// and there isn't one of the same name on the mesh, we can customize it
+					if (SocketItem->CanCustomizeSocket() )
 					{
-						// If a socket is on the mesh only, then offer to promote it to the skeleton
-						MenuBuilder.AddMenuEntry( Actions.PromoteSocketToSkeleton );
+						if (SocketItem->GetParentType() == ESocketParentType::Skeleton )
+						{
+							MenuBuilder.AddMenuEntry( Actions.CreateMeshSocket );
+						}
+						else if (SocketItem->GetParentType() == ESocketParentType::Mesh )
+						{
+							// If a socket is on the mesh only, then offer to promote it to the skeleton
+							MenuBuilder.AddMenuEntry( Actions.PromoteSocketToSkeleton );
+						}
 					}
 				}
-			}
 
-			MenuBuilder.EndSection();
+				MenuBuilder.EndSection();
+			}
 		}
 
-		MenuBuilder.BeginSection("SkeletonTreeAttachedAssets", LOCTEXT( "AttachedAssetsActionsHeader", "Attached Assets Actions" ) );
-
-		if ( BoneTreeSelection.IsSingleItemSelected() )
+		if (BoneTreeSelection.HasSelectedOfType<FSkeletonTreeBoneItem>() || BoneTreeSelection.HasSelectedOfType<FSkeletonTreeSocketItem>())
 		{
-			MenuBuilder.AddSubMenu(	LOCTEXT( "AttachNewAsset", "Add Preview Asset" ),
-				LOCTEXT ( "AttachNewAsset_ToolTip", "Attaches an asset to this part of the skeleton. Assets can also be dragged onto the skeleton from a content browser to attach" ),
-				FNewMenuDelegate::CreateSP( this, &SSkeletonTree::FillAttachAssetSubmenu, BoneTreeSelection.GetSingleSelectedItem() ) );
-		}
+			MenuBuilder.BeginSection("SkeletonTreeAttachedAssets", LOCTEXT( "AttachedAssetsActionsHeader", "Attached Assets Actions" ) );
 
-		FUIAction RemoveAllAttachedAssets = FUIAction(	FExecuteAction::CreateSP( this, &SSkeletonTree::OnRemoveAllAssets ),
-			FCanExecuteAction::CreateSP( this, &SSkeletonTree::CanRemoveAllAssets ));
+			if ( BoneTreeSelection.IsSingleItemSelected() )
+			{
+				MenuBuilder.AddSubMenu(	LOCTEXT( "AttachNewAsset", "Add Preview Asset" ),
+					LOCTEXT ( "AttachNewAsset_ToolTip", "Attaches an asset to this part of the skeleton. Assets can also be dragged onto the skeleton from a content browser to attach" ),
+					FNewMenuDelegate::CreateSP( this, &SSkeletonTree::FillAttachAssetSubmenu, BoneTreeSelection.GetSingleSelectedItem() ) );
+			}
+
+			FUIAction RemoveAllAttachedAssets = FUIAction(	FExecuteAction::CreateSP( this, &SSkeletonTree::OnRemoveAllAssets ),
+				FCanExecuteAction::CreateSP( this, &SSkeletonTree::CanRemoveAllAssets ));
 
-		MenuBuilder.AddMenuEntry( LOCTEXT( "RemoveAllAttachedAssets", "Remove All Attached Assets" ),
-			LOCTEXT ( "RemoveAllAttachedAssets_ToolTip", "Removes all the attached assets from the skeleton and mesh." ),
-			FSlateIcon(), RemoveAllAttachedAssets );
+			MenuBuilder.AddMenuEntry( LOCTEXT( "RemoveAllAttachedAssets", "Remove All Attached Assets" ),
+				LOCTEXT ( "RemoveAllAttachedAssets_ToolTip", "Removes all the attached assets from the skeleton and mesh." ),
+				FSlateIcon(), RemoveAllAttachedAssets );
+
+			MenuBuilder.EndSection();
+		}
 
+		// Add an empty section so the menu can be extended when there are no optionally-added entries
+		MenuBuilder.BeginSection("SkeletonTreeContextMenu");
 		MenuBuilder.EndSection();
 	}
 
@@ -881,8 +948,9 @@ void SSkeletonTree::CreateMenuForBoneReduction(FMenuBuilder& MenuBuilder, SSkele
 void SSkeletonTree::SetBoneTranslationRetargetingModeRecursive(EBoneTranslationRetargetingMode::Type NewRetargetingMode)
 {
 	TArray<FName> BoneNames;
-	FBoneTreeSelection TreeSelection(SkeletonTreeView->GetSelectedItems());
-	for (const TSharedPtr<FSkeletonTreeBoneItem>& Item : TreeSelection.SelectedBones)
+	TArray<TSharedPtr<ISkeletonTreeItem>> SelectedItems = SkeletonTreeView->GetSelectedItems();
+	FSkeletonTreeSelection TreeSelection(SelectedItems);
+	for (const TSharedPtr<FSkeletonTreeBoneItem>& Item : TreeSelection.GetSelectedItems<FSkeletonTreeBoneItem>())
 	{
 		BoneNames.Add(Item->GetRowItemName());
 	}
@@ -908,12 +976,13 @@ void SSkeletonTree::RemoveFromLOD(int32 LODIndex, bool bIncludeSelected, bool bI
 	const FText Message(LOCTEXT("RemoveBonesFromLODWarning", "This action can't be undone. Would you like to continue?"));
 	if (FMessageDialog::Open(EAppMsgType::YesNo, Message) == EAppReturnType::Yes)
 	{
-		FBoneTreeSelection TreeSelection(SkeletonTreeView->GetSelectedItems());
+		TArray<TSharedPtr<ISkeletonTreeItem>> SelectedItems = SkeletonTreeView->GetSelectedItems();
+		FSkeletonTreeSelection TreeSelection(SelectedItems);
 		const FReferenceSkeleton& RefSkeleton = GetEditableSkeletonInternal()->GetSkeleton().GetReferenceSkeleton();
 
 		TArray<FName> BonesToRemove;
 
-		for (const TSharedPtr<FSkeletonTreeBoneItem>& Item : TreeSelection.SelectedBones)
+		for (const TSharedPtr<FSkeletonTreeBoneItem>& Item : TreeSelection.GetSelectedItems<FSkeletonTreeBoneItem>())
 		{
 			FName BoneName = Item->GetRowItemName();
 			int32 BoneIndex = RefSkeleton.FindBoneIndex(BoneName);
@@ -960,13 +1029,14 @@ void SSkeletonTree::RemoveFromLOD(int32 LODIndex, bool bIncludeSelected, bool bI
 
 void SSkeletonTree::OnCopyBoneNames()
 {
-	FBoneTreeSelection TreeSelection(SkeletonTreeView->GetSelectedItems());
-
-	if( TreeSelection.SelectedBones.Num() > 0 )
+	TArray<TSharedPtr<ISkeletonTreeItem>> SelectedItems = SkeletonTreeView->GetSelectedItems();
+	FSkeletonTreeSelection TreeSelection(SelectedItems);
+	TArray<TSharedPtr<FSkeletonTreeBoneItem>> SelectedBones = TreeSelection.GetSelectedItems<FSkeletonTreeBoneItem>();
+	if( SelectedBones.Num() > 0 )
 	{
 		bool bFirst = true;
 		FString BoneNames;
-		for (const TSharedPtr<FSkeletonTreeBoneItem>& Item : TreeSelection.SelectedBones)
+		for (const TSharedPtr<FSkeletonTreeBoneItem>& Item : SelectedBones)
 		{
 			FName BoneName = Item->GetRowItemName();
 			if (!bFirst)
@@ -989,14 +1059,15 @@ void SSkeletonTree::OnResetBoneTransforms()
 		UAnimPreviewInstance* PreviewInstance = PreviewComponent->PreviewInstance;
 		check(PreviewInstance);
 
-		FBoneTreeSelection TreeSelection(SkeletonTreeView->GetSelectedItems());
-
-		if (TreeSelection.SelectedBones.Num() > 0)
+		TArray<TSharedPtr<ISkeletonTreeItem>> SelectedItems = SkeletonTreeView->GetSelectedItems();
+		FSkeletonTreeSelection TreeSelection(SelectedItems);
+		TArray<TSharedPtr<FSkeletonTreeBoneItem>> SelectedBones = TreeSelection.GetSelectedItems<FSkeletonTreeBoneItem>();
+		if (SelectedBones.Num() > 0)
 		{
 			bool bModified = false;
 			GEditor->BeginTransaction(LOCTEXT("SkeletonTree_ResetBoneTransforms", "Reset Bone Transforms"));
 
-			for (const TSharedPtr<FSkeletonTreeBoneItem>& Item : TreeSelection.SelectedBones)
+			for (const TSharedPtr<FSkeletonTreeBoneItem>& Item : SelectedBones)
 			{
 				FName BoneName = Item->GetRowItemName();
 				const FAnimNode_ModifyBone* ModifiedBone = PreviewInstance->FindModifiedBone(BoneName);
@@ -1020,14 +1091,15 @@ void SSkeletonTree::OnResetBoneTransforms()
 
 void SSkeletonTree::OnCopySockets() const
 {
-	FBoneTreeSelection TreeSelection(SkeletonTreeView->GetSelectedItems());
-
-	int32 NumSocketsToCopy = TreeSelection.SelectedSockets.Num();
+	TArray<TSharedPtr<ISkeletonTreeItem>> SelectedItems = SkeletonTreeView->GetSelectedItems();
+	FSkeletonTreeSelection TreeSelection(SelectedItems);
+	TArray<TSharedPtr<FSkeletonTreeSocketItem>> SelectedSockets = TreeSelection.GetSelectedItems<FSkeletonTreeSocketItem>();
+	int32 NumSocketsToCopy = SelectedSockets.Num();
 	if ( NumSocketsToCopy > 0 )
 	{
 		FString SocketsDataString;
 
-		for (const TSharedPtr<FSkeletonTreeSocketItem>& Item : TreeSelection.SelectedSockets)
+		for (const TSharedPtr<FSkeletonTreeSocketItem>& Item : SelectedSockets)
 		{
 			SocketsDataString += SerializeSocketToString(Item->GetSocket(), Item->GetParentType());
 		}
@@ -1054,7 +1126,8 @@ FString SSkeletonTree::SerializeSocketToString( USkeletalMeshSocket* Socket, ESo
 
 void SSkeletonTree::OnPasteSockets(bool bPasteToSelectedBone)
 {
-	FBoneTreeSelection TreeSelection(SkeletonTreeView->GetSelectedItems());
+	TArray<TSharedPtr<ISkeletonTreeItem>> SelectedItems = SkeletonTreeView->GetSelectedItems();
+	FSkeletonTreeSelection TreeSelection(SelectedItems);
 
 	// Pasting sockets should only work if there is just one bone selected
 	if ( TreeSelection.IsSingleOfTypeSelected<FSkeletonTreeBoneItem>() )
@@ -1070,7 +1143,8 @@ void SSkeletonTree::OnPasteSockets(bool bPasteToSelectedBone)
 void SSkeletonTree::OnAddSocket()
 {
 	// This adds a socket to the currently selected bone in the SKELETON, not the MESH.
-	FBoneTreeSelection TreeSelection(SkeletonTreeView->GetSelectedItems());
+	TArray<TSharedPtr<ISkeletonTreeItem>> SelectedItems = SkeletonTreeView->GetSelectedItems();
+	FSkeletonTreeSelection TreeSelection(SelectedItems);
 
 	// Can only add a socket to one bone
 	if (TreeSelection.IsSingleOfTypeSelected<FSkeletonTreeBoneItem>())
@@ -1101,7 +1175,8 @@ void SSkeletonTree::OnCustomizeSocket()
 {
 	// This should only be called on a skeleton socket, it copies the 
 	// socket to the mesh so the user can edit it separately
-	FBoneTreeSelection TreeSelection(SkeletonTreeView->GetSelectedItems());
+	TArray<TSharedPtr<ISkeletonTreeItem>> SelectedItems = SkeletonTreeView->GetSelectedItems();
+	FSkeletonTreeSelection TreeSelection(SelectedItems);
 
 	if(TreeSelection.IsSingleOfTypeSelected<FSkeletonTreeSocketItem>())
 	{
@@ -1116,7 +1191,8 @@ void SSkeletonTree::OnPromoteSocket()
 {
 	// This should only be called on a mesh socket, it copies the 
 	// socket to the skeleton so all meshes can use it
-	FBoneTreeSelection TreeSelection(SkeletonTreeView->GetSelectedItems());
+	TArray<TSharedPtr<ISkeletonTreeItem>> SelectedItems = SkeletonTreeView->GetSelectedItems();
+	FSkeletonTreeSelection TreeSelection(SelectedItems);
 
 	// Can only customize one socket (CreateContextMenu() should prevent this firing!)
 	if(TreeSelection.IsSingleOfTypeSelected<FSkeletonTreeSocketItem>())
@@ -1185,13 +1261,15 @@ bool SSkeletonTree::CanRemoveAllAssets() const
 
 bool SSkeletonTree::CanRenameSelected() const
 {
-	FBoneTreeSelection TreeSelection(SkeletonTreeView->GetSelectedItems());
+	TArray<TSharedPtr<ISkeletonTreeItem>> SelectedItems = SkeletonTreeView->GetSelectedItems();
+	FSkeletonTreeSelection TreeSelection(SelectedItems);
 	return TreeSelection.IsSingleOfTypeSelected<FSkeletonTreeSocketItem>() || TreeSelection.IsSingleOfTypeSelected<FSkeletonTreeVirtualBoneItem>();
 }
 
 void SSkeletonTree::OnRenameSelected()
 {
-	FBoneTreeSelection TreeSelection(SkeletonTreeView->GetSelectedItems());
+	TArray<TSharedPtr<ISkeletonTreeItem>> SelectedItems = SkeletonTreeView->GetSelectedItems();
+	FSkeletonTreeSelection TreeSelection(SelectedItems);
 
 	if(TreeSelection.IsSingleOfTypeSelected<FSkeletonTreeSocketItem>() || TreeSelection.IsSingleOfTypeSelected<FSkeletonTreeVirtualBoneItem>())
 	{
@@ -1202,10 +1280,25 @@ void SSkeletonTree::OnRenameSelected()
 
 void SSkeletonTree::OnSelectionChanged(TSharedPtr<ISkeletonTreeItem> Selection, ESelectInfo::Type SelectInfo)
 {
+	TArray<TSharedPtr<ISkeletonTreeItem>> SelectedItems = SkeletonTreeView->GetSelectedItems();
+
 	if( Selection.IsValid() )
 	{
+		// Disable bone proxy ticking on all bone/virtual bones
+		for (TSharedPtr<ISkeletonTreeItem>& Item : LinearItems)
+		{
+			if (Item->IsOfType<FSkeletonTreeBoneItem>())
+			{
+				StaticCastSharedPtr<FSkeletonTreeBoneItem>(Item)->EnableBoneProxyTick(false);
+			}
+			else if (Item->IsOfType<FSkeletonTreeVirtualBoneItem>())
+			{
+				StaticCastSharedPtr<FSkeletonTreeVirtualBoneItem>(Item)->EnableBoneProxyTick(false);
+			}
+		}
+
 		//Get all the selected items
-		FBoneTreeSelection TreeSelection(SkeletonTreeView->GetSelectedItems());
+		FSkeletonTreeSelection TreeSelection(SelectedItems);
 
 		if (GetPreviewScene().IsValid())
 		{
@@ -1213,23 +1306,37 @@ void SSkeletonTree::OnSelectionChanged(TSharedPtr<ISkeletonTreeItem> Selection,
 			if (TreeSelection.SelectedItems.Num() > 0 && PreviewComponent)
 			{
 				// pick the first settable bone from the selection
-				for (auto ItemIt = TreeSelection.SelectedItems.CreateConstIterator(); ItemIt; ++ItemIt)
+				for (TSharedPtr<ISkeletonTreeItem> Item : TreeSelection.SelectedItems)
 				{
-					TSharedPtr<ISkeletonTreeItem> Item = *(ItemIt);
-
-					// Test SelectInfo so we don't end up in an infinite loop due to delegates calling each other
-					if (SelectInfo != ESelectInfo::Direct && (Item->IsOfType<FSkeletonTreeBoneItem>() || Item->IsOfType<FSkeletonTreeVirtualBoneItem>()))
+					if((Item->IsOfType<FSkeletonTreeBoneItem>() || Item->IsOfType<FSkeletonTreeVirtualBoneItem>()))
 					{
-						FName BoneName = Item->GetRowItemName();
+						// enable ticking on the selected bone proxies
+						if (Item->IsOfType<FSkeletonTreeBoneItem>())
+						{
+							StaticCastSharedPtr<FSkeletonTreeBoneItem>(Item)->EnableBoneProxyTick(true);
+						}
+						else if (Item->IsOfType<FSkeletonTreeVirtualBoneItem>())
+						{
+							StaticCastSharedPtr<FSkeletonTreeVirtualBoneItem>(Item)->EnableBoneProxyTick(true);
+						}
 
-						// Get bone index
-						int32 BoneIndex = PreviewComponent->GetBoneIndex(BoneName);
-						if (BoneIndex != INDEX_NONE)
+						// Test SelectInfo so we don't end up in an infinite loop due to delegates calling each other
+						if (SelectInfo != ESelectInfo::Direct)
 						{
-							GetPreviewScene()->SetSelectedBone(BoneName);
-							BoneProxy->BoneName = BoneName;
-							OnObjectSelectedMulticast.Broadcast(BoneProxy);
-							break;
+							FName BoneName = Item->GetRowItemName();
+
+							// Get bone index
+							int32 BoneIndex = PreviewComponent->GetBoneIndex(BoneName);
+							if (BoneIndex != INDEX_NONE)
+							{
+								GetPreviewScene()->SetSelectedBone(BoneName);
+								BoneProxy->BoneName = BoneName;
+
+								PRAGMA_DISABLE_DEPRECATION_WARNINGS
+								OnObjectSelectedMulticast.Broadcast(BoneProxy);
+								PRAGMA_ENABLE_DEPRECATION_WARNINGS
+								break;
+							}
 						}
 					}
 					// Test SelectInfo so we don't end up in an infinite loop due to delegates calling each other
@@ -1240,13 +1347,17 @@ void SSkeletonTree::OnSelectionChanged(TSharedPtr<ISkeletonTreeItem> Selection,
 						FSelectedSocketInfo SocketInfo(Socket, SocketItem->GetParentType() == ESocketParentType::Skeleton);
 						GetPreviewScene()->SetSelectedSocket(SocketInfo);
 
+						PRAGMA_DISABLE_DEPRECATION_WARNINGS
 						OnObjectSelectedMulticast.Broadcast(SocketInfo.Socket);
+						PRAGMA_ENABLE_DEPRECATION_WARNINGS
 					}
 					else if (Item->IsOfType<FSkeletonTreeAttachedAssetItem>())
 					{
 						GetPreviewScene()->DeselectAll();
 
+						PRAGMA_DISABLE_DEPRECATION_WARNINGS
 						OnObjectSelectedMulticast.Broadcast(nullptr);
+						PRAGMA_ENABLE_DEPRECATION_WARNINGS
 					}
 				}
 				PreviewComponent->PostInitMeshObject(PreviewComponent->MeshObject);
@@ -1261,8 +1372,13 @@ void SSkeletonTree::OnSelectionChanged(TSharedPtr<ISkeletonTreeItem> Selection,
 			GetPreviewScene()->DeselectAll();
 		}
 
+		PRAGMA_DISABLE_DEPRECATION_WARNINGS
 		OnObjectSelectedMulticast.Broadcast(nullptr);
+		PRAGMA_ENABLE_DEPRECATION_WARNINGS
 	}
+
+	TArrayView<TSharedPtr<ISkeletonTreeItem>> ArrayView(SelectedItems);
+	OnSelectionChangedMulticast.Broadcast(ArrayView, SelectInfo);
 }
 
 void SSkeletonTree::AttachAssets(const TSharedRef<ISkeletonTreeItem>& TargetItem, const TArray<FAssetData>& AssetData)
@@ -1321,12 +1437,18 @@ void SSkeletonTree::SetTreeItemExpansionRecursive(TSharedPtr< ISkeletonTreeItem
 	}
 }
 
-void SSkeletonTree::PostUndo()
+void SSkeletonTree::PostUndo(bool bSuccess)
 {
 	// Rebuild the tree view whenever we undo a change to the skeleton
 	CreateTreeColumns();
 }
 
+void SSkeletonTree::PostRedo(bool bSuccess)
+{
+	// Rebuild the tree view whenever we redo a change to the skeleton
+	CreateTreeColumns();
+}
+
 void SSkeletonTree::OnFilterTextChanged( const FText& SearchText )
 {
 	FilterText = SearchText;
@@ -1334,6 +1456,16 @@ void SSkeletonTree::OnFilterTextChanged( const FText& SearchText )
 	ApplyFilter();
 }
 
+void SSkeletonTree::Refresh()
+{
+	CreateFromSkeleton();
+}
+
+void SSkeletonTree::RefreshFilter()
+{
+	ApplyFilter();
+}
+
 void SSkeletonTree::SetSkeletalMesh(USkeletalMesh* NewSkeletalMesh)
 {
 	if (GetPreviewScene().IsValid())
@@ -1370,7 +1502,9 @@ void SSkeletonTree::SetSelectedSocket( const FSelectedSocketInfo& SocketInfo )
 			GetPreviewScene()->SetSelectedSocket(SocketInfo);
 		}
 
+		PRAGMA_DISABLE_DEPRECATION_WARNINGS
 		OnObjectSelectedMulticast.Broadcast(SocketInfo.Socket);
+		PRAGMA_ENABLE_DEPRECATION_WARNINGS
 	}
 }
 
@@ -1400,7 +1534,9 @@ void SSkeletonTree::SetSelectedBone( const FName& BoneName )
 		}
 
 		BoneProxy->BoneName = BoneName;
+		PRAGMA_DISABLE_DEPRECATION_WARNINGS
 		OnObjectSelectedMulticast.Broadcast(BoneProxy);
+		PRAGMA_ENABLE_DEPRECATION_WARNINGS
 	}
 }
 
@@ -1416,7 +1552,9 @@ void SSkeletonTree::DeselectAll()
 			GetPreviewScene()->DeselectAll();
 		}
 
+		PRAGMA_DISABLE_DEPRECATION_WARNINGS
 		OnObjectSelectedMulticast.Broadcast(nullptr);
+		PRAGMA_ENABLE_DEPRECATION_WARNINGS
 	}
 }
 
@@ -1434,34 +1572,44 @@ TSharedRef< SWidget > SSkeletonTree::CreateFilterMenu()
 	const FSkeletonTreeCommands& Actions = FSkeletonTreeCommands::Get();
 
 	const bool CloseAfterSelection = true;
-	FMenuBuilder MenuBuilder( CloseAfterSelection, UICommandList );
+	FMenuBuilder MenuBuilder( CloseAfterSelection, UICommandList, Extenders );
 
-	MenuBuilder.BeginSection("Bones", LOCTEXT("OptionsMenuHeading", "Options"));
+	MenuBuilder.BeginSection("FilterOptions", LOCTEXT("OptionsMenuHeading", "Options"));
 	{
-		MenuBuilder.AddMenuEntry(Actions.ShowRetargeting);
+		if(Builder->IsShowingBones() && bAllowSkeletonOperations)
+		{
+			MenuBuilder.AddMenuEntry(Actions.ShowRetargeting);
+		}
 		MenuBuilder.AddMenuEntry(Actions.FilteringFlattensHierarchy);
+		MenuBuilder.AddMenuEntry(Actions.HideParentsWhenFiltering);
 	}
 	MenuBuilder.EndSection();
 
-	MenuBuilder.BeginSection("Bones", LOCTEXT( "BonesMenuHeading", "Bones" ) );
+	if(Builder->IsShowingBones())
 	{
-		MenuBuilder.AddMenuEntry( Actions.ShowAllBones );
-		MenuBuilder.AddMenuEntry( Actions.ShowMeshBones );
-		MenuBuilder.AddMenuEntry(Actions.ShowLODBones);
-		MenuBuilder.AddMenuEntry( Actions.ShowWeightedBones );
-		MenuBuilder.AddMenuEntry( Actions.HideBones );
+		MenuBuilder.BeginSection("FilterBones", LOCTEXT( "BonesMenuHeading", "Bones" ) );
+		{
+			MenuBuilder.AddMenuEntry( Actions.ShowAllBones );
+			MenuBuilder.AddMenuEntry( Actions.ShowMeshBones );
+			MenuBuilder.AddMenuEntry( Actions.ShowLODBones);
+			MenuBuilder.AddMenuEntry( Actions.ShowWeightedBones );
+			MenuBuilder.AddMenuEntry( Actions.HideBones );
+		}
+		MenuBuilder.EndSection();
 	}
-	MenuBuilder.EndSection();
 
-	MenuBuilder.BeginSection("Sockets", LOCTEXT("SocketsMenuHeading", "Sockets"));
+	if(Builder->IsShowingSockets())
 	{
-		MenuBuilder.AddMenuEntry(Actions.ShowActiveSockets);
-		MenuBuilder.AddMenuEntry(Actions.ShowMeshSockets);
-		MenuBuilder.AddMenuEntry(Actions.ShowSkeletonSockets);
-		MenuBuilder.AddMenuEntry(Actions.ShowAllSockets);
-		MenuBuilder.AddMenuEntry(Actions.HideSockets);
+		MenuBuilder.BeginSection("FilterSockets", LOCTEXT("SocketsMenuHeading", "Sockets"));
+		{
+			MenuBuilder.AddMenuEntry(Actions.ShowActiveSockets);
+			MenuBuilder.AddMenuEntry(Actions.ShowMeshSockets);
+			MenuBuilder.AddMenuEntry(Actions.ShowSkeletonSockets);
+			MenuBuilder.AddMenuEntry(Actions.ShowAllSockets);
+			MenuBuilder.AddMenuEntry(Actions.HideSockets);
+		}
+		MenuBuilder.EndSection();
 	}
-	MenuBuilder.EndSection();
 
 	return MenuBuilder.MakeWidget();
 }
@@ -1513,67 +1661,85 @@ bool SSkeletonTree::IsSocketFilter( ESocketFilter InSocketFilter ) const
 
 FText SSkeletonTree::GetFilterMenuTitle() const
 {
-	FText BoneFilterMenuText;
+	TArray<FText> FilterLabels;
 
-	switch ( BoneFilter )
+	if(Builder->IsShowingBones())	
 	{
-	case EBoneFilter::All:
-		BoneFilterMenuText = LOCTEXT( "BoneFilterMenuAll", "All Bones" );
-		break;
+		switch ( BoneFilter )
+		{
+		case EBoneFilter::All:
+			FilterLabels.Add(LOCTEXT( "BoneFilterMenuAll", "Bones" ));
+			break;
 
-	case EBoneFilter::Mesh:
-		BoneFilterMenuText = LOCTEXT( "BoneFilterMenuMesh", "Mesh Bones" );
-		break;
+		case EBoneFilter::Mesh:
+			FilterLabels.Add(LOCTEXT( "BoneFilterMenuMesh", "Mesh Bones" ));
+			break;
 
-	case EBoneFilter::LOD:
-		BoneFilterMenuText = LOCTEXT("BoneFilterMenuLOD", "LOD Bones");
-		break;
+		case EBoneFilter::LOD:
+			FilterLabels.Add(LOCTEXT("BoneFilterMenuLOD", "LOD Bones"));
+			break;
 
-	case EBoneFilter::Weighted:
-		BoneFilterMenuText = LOCTEXT( "BoneFilterMenuWeighted", "Weighted Bones" );
-		break;
+		case EBoneFilter::Weighted:
+			FilterLabels.Add(LOCTEXT( "BoneFilterMenuWeighted", "Weighted Bones" ));
+			break;
 
-	case EBoneFilter::None:
-		BoneFilterMenuText = LOCTEXT( "BoneFilterMenuHidden", "Bones Hidden" );
-		break;
+		case EBoneFilter::None:
+			break;
 
-	default:
-		// Unknown mode
-		check( 0 );
-		break;
+		default:
+			// Unknown mode
+			check(false);
+			break;
+		}
 	}
 
-	FText SocketFilterMenuText;
-
-	switch (SocketFilter)
+	if(Builder->IsShowingSockets())
 	{
-	case ESocketFilter::Active:
-		SocketFilterMenuText = LOCTEXT("SocketFilterMenuActive", "Active Sockets");
-		break;
+		switch (SocketFilter)
+		{
+		case ESocketFilter::Active:
+			FilterLabels.Add(LOCTEXT("SocketFilterMenuActive", "Active Sockets"));
+			break;
 
-	case ESocketFilter::Mesh:
-		SocketFilterMenuText = LOCTEXT("SocketFilterMenuMesh", "Mesh Sockets");
-		break;
+		case ESocketFilter::Mesh:
+			FilterLabels.Add(LOCTEXT("SocketFilterMenuMesh", "Mesh Sockets"));
+			break;
 
-	case ESocketFilter::Skeleton:
-		SocketFilterMenuText = LOCTEXT("SocketFilterMenuSkeleton", "Skeleton Sockets");
-		break;
+		case ESocketFilter::Skeleton:
+			FilterLabels.Add(LOCTEXT("SocketFilterMenuSkeleton", "Skeleton Sockets"));
+			break;
 
-	case ESocketFilter::All:
-		SocketFilterMenuText = LOCTEXT("SocketFilterMenuAll", "All Sockets");
-		break;
+		case ESocketFilter::All:
+			FilterLabels.Add(LOCTEXT("SocketFilterMenuAll", "All Sockets"));
+			break;
 
-	case ESocketFilter::None:
-		SocketFilterMenuText = LOCTEXT("SocketFilterMenuHidden", "Sockets Hidden");
-		break;
+		case ESocketFilter::None:
+			break;
 
-	default:
-		// Unknown mode
-		check(0);
-		break;
+		default:
+			// Unknown mode
+			check(false);
+			break;
+		}
 	}
 
-	return FText::Format(LOCTEXT("FilterMenuLabelFormat", "{0}, {1}"), BoneFilterMenuText, SocketFilterMenuText);
+	OnGetFilterText.ExecuteIfBound(FilterLabels);
+
+	FText Label;
+	if(FilterLabels.Num() > 0)
+	{
+		Label = FText::Format(LOCTEXT("FilterMenuLabelFormatStart", "{0}"), FilterLabels[0]);
+		for(int32 LabelIndex = 1; LabelIndex < FilterLabels.Num(); ++LabelIndex)
+		{
+			Label = FText::Format(LOCTEXT("FilterMenuLabelFormat", "{0}, {1}"), Label, FilterLabels[LabelIndex]);
+		}
+	}
+	else
+	{
+		Label = LOCTEXT("ShowingNoneLabel", "Filters");
+	}
+	
+	return Label;
 }
 
 bool SSkeletonTree::IsAddingSocketsAllowed() const
@@ -1598,17 +1764,25 @@ FReply SSkeletonTree::OnKeyDown( const FGeometry& MyGeometry, const FKeyEvent& I
 	return FReply::Unhandled();
 }
 
+bool SSkeletonTree::CanDeleteSelectedRows() const
+{
+	TArray<TSharedPtr<ISkeletonTreeItem>> SelectedItems = SkeletonTreeView->GetSelectedItems();
+	FSkeletonTreeSelection TreeSelection(SelectedItems);
+	return (TreeSelection.HasSelectedOfType<FSkeletonTreeAttachedAssetItem>() || TreeSelection.HasSelectedOfType<FSkeletonTreeSocketItem>() || TreeSelection.HasSelectedOfType<FSkeletonTreeVirtualBoneItem>());
+}
+
 void SSkeletonTree::OnDeleteSelectedRows()
 {
-	FBoneTreeSelection TreeSelection(SkeletonTreeView->GetSelectedItems());
+	TArray<TSharedPtr<ISkeletonTreeItem>> SelectedItems = SkeletonTreeView->GetSelectedItems();
+	FSkeletonTreeSelection TreeSelection(SelectedItems);
 
 	if(TreeSelection.HasSelectedOfType<FSkeletonTreeAttachedAssetItem>() || TreeSelection.HasSelectedOfType<FSkeletonTreeSocketItem>() || TreeSelection.HasSelectedOfType<FSkeletonTreeVirtualBoneItem>())
 	{
 		FScopedTransaction Transaction( LOCTEXT( "SkeletonTreeDeleteSelected", "Delete selected sockets/meshes/bones from skeleton tree" ) );
 
-		DeleteAttachedAssets( TreeSelection.SelectedAssets );
-		DeleteSockets( TreeSelection.SelectedSockets );
-		DeleteVirtualBones( TreeSelection.SelectedVirtualBones );
+		DeleteAttachedAssets( TreeSelection.GetSelectedItems<FSkeletonTreeAttachedAssetItem>() );
+		DeleteSockets( TreeSelection.GetSelectedItems<FSkeletonTreeSocketItem>() );
+		DeleteVirtualBones( TreeSelection.GetSelectedItems<FSkeletonTreeVirtualBoneItem>() );
 
 		CreateFromSkeleton();
 	}
@@ -1692,9 +1866,10 @@ void SSkeletonTree::RecursiveSetBlendProfileScales(float InScaleToSet)
 	if(SelectedBlendProfile)
 	{
 		TArray<FName> BoneNames;
-		FBoneTreeSelection TreeSelection(SkeletonTreeView->GetSelectedItems());
+		TArray<TSharedPtr<ISkeletonTreeItem>> SelectedItems = SkeletonTreeView->GetSelectedItems();
+		FSkeletonTreeSelection TreeSelection(SelectedItems);
 
-		for(TSharedPtr<FSkeletonTreeBoneItem>& SelectedBone : TreeSelection.SelectedBones)
+		for(TSharedPtr<FSkeletonTreeBoneItem>& SelectedBone : TreeSelection.GetSelectedItems<FSkeletonTreeBoneItem>())
 		{
 			BoneNames.Add(SelectedBone->GetRowItemName());
 		}
@@ -1769,6 +1944,43 @@ void SSkeletonTree::OnLODSwitched()
 	}
 }
 
+void SSkeletonTree::SelectItemsBy(TFunctionRef<bool(const TSharedRef<ISkeletonTreeItem>&, bool&)> Predicate) const
+{
+	SkeletonTreeView->ClearSelection();
+
+	TSharedPtr<ISkeletonTreeItem> ScrollItem = nullptr;
+	for (const TSharedPtr<ISkeletonTreeItem>& Item : LinearItems)
+	{
+		bool bExpand = false;
+		if (Predicate(Item.ToSharedRef(), bExpand))
+		{
+			SkeletonTreeView->SetItemSelection(Item, true);
+			if (bExpand)
+			{
+				if(Item->GetChildren().Num() == 0)
+				{
+					// leaf nodes expand their parent
+					TSharedPtr<ISkeletonTreeItem> ParentItem = Item->GetParent();
+					if(ParentItem.IsValid())
+					{
+						SkeletonTreeView->SetItemExpansion(ParentItem, true);
+					}
+				}
+				else
+				{
+					SkeletonTreeView->SetItemExpansion(Item, true);
+				}
+			}
+			ScrollItem = Item;
+		}
+	}
+
+	if (ScrollItem.IsValid())
+	{
+		SkeletonTreeView->RequestScrollIntoView(ScrollItem);
+	}
+}
+
 void SSkeletonTree::DuplicateAndSelectSocket(const FSelectedSocketInfo& SocketInfoToDuplicate, const FName& NewParentBoneName /*= FName()*/)
 {
 	USkeletalMesh* SkeletalMesh = GetPreviewScene().IsValid() ? GetPreviewScene()->GetPreviewMeshComponent()->SkeletalMesh : nullptr;
@@ -1805,107 +2017,88 @@ void SSkeletonTree::HandleFocusCamera()
 	}
 }
 
-/** Filter utility class */
-class FSkeletonTreeFilterContext : public ITextFilterExpressionContext
-{
-public:
-	explicit FSkeletonTreeFilterContext(const FName& InName)
-		: Name(InName)
-	{
-	}
-
-	virtual bool TestBasicStringExpression(const FTextFilterString& InValue, const ETextFilterTextComparisonMode InTextComparisonMode) const override
-	{
-		return TextFilterUtils::TestBasicStringExpression(Name.ToString(), InValue, InTextComparisonMode);
-	}
-
-	virtual bool TestComplexExpression(const FName& InKey, const FTextFilterString& InValue, const ETextFilterComparisonOperation InComparisonOperation, const ETextFilterTextComparisonMode InTextComparisonMode) const override
-	{
-		return false;
-	}
-
-private:
-	FName Name;
-};
-
-ESkeletonTreeFilterResult SSkeletonTree::HandleFilterSkeletonTreeItem(const TSharedPtr<class ISkeletonTreeItem>& InItem)
+ESkeletonTreeFilterResult SSkeletonTree::HandleFilterSkeletonTreeItem(const FSkeletonTreeFilterArgs& InArgs, const TSharedPtr<class ISkeletonTreeItem>& InItem)
 {
 	ESkeletonTreeFilterResult Result = ESkeletonTreeFilterResult::Shown;
 
-	if (!FilterText.IsEmpty())
+	if(InItem->IsOfType<FSkeletonTreeBoneItem>() || InItem->IsOfType<FSkeletonTreeSocketItem>() || InItem->IsOfType<FSkeletonTreeAttachedAssetItem>() || InItem->IsOfType<FSkeletonTreeVirtualBoneItem>())
 	{
-		if (TextFilterPtr->TestTextFilter(FSkeletonTreeFilterContext(InItem->GetRowItemName())))
-		{
-			Result = ESkeletonTreeFilterResult::ShownHighlighted;
-		}
-		else
+
+		if (InArgs.TextFilter.IsValid())
 		{
-			Result = ESkeletonTreeFilterResult::Hidden;
+			if (InArgs.TextFilter->TestTextFilter(FSkeletonTreeFilterContext(InItem->GetRowItemName())))
+			{
+				Result = ESkeletonTreeFilterResult::ShownHighlighted;
+			}
+			else
+			{
+				Result = ESkeletonTreeFilterResult::Hidden;
+			}
 		}
-	}
 
 
-	if (InItem->IsOfType<FSkeletonTreeBoneItem>())
-	{
-		TSharedPtr<FSkeletonTreeBoneItem> BoneItem = StaticCastSharedPtr<FSkeletonTreeBoneItem>(InItem);
-
-		if (BoneFilter == EBoneFilter::None)
-		{
-			Result = ESkeletonTreeFilterResult::Hidden;
-		}
-		else if (GetPreviewScene().IsValid())
+		if (InItem->IsOfType<FSkeletonTreeBoneItem>())
 		{
-			UDebugSkelMeshComponent* PreviewMeshComponent = GetPreviewScene()->GetPreviewMeshComponent();
-			if (PreviewMeshComponent)
-			{
-				int32 BoneMeshIndex = PreviewMeshComponent->GetBoneIndex(BoneItem->GetRowItemName());
+			TSharedPtr<FSkeletonTreeBoneItem> BoneItem = StaticCastSharedPtr<FSkeletonTreeBoneItem>(InItem);
 
-				// Remove non-mesh bones if we're filtering
-				if ((BoneFilter == EBoneFilter::Mesh || BoneFilter == EBoneFilter::Weighted || BoneFilter == EBoneFilter::LOD) &&
-					BoneMeshIndex == INDEX_NONE)
+			if (BoneFilter == EBoneFilter::None)
+			{
+				Result = ESkeletonTreeFilterResult::Hidden;
+			}
+			else if (GetPreviewScene().IsValid())
+			{
+				UDebugSkelMeshComponent* PreviewMeshComponent = GetPreviewScene()->GetPreviewMeshComponent();
+				if (PreviewMeshComponent)
 				{
-					Result = ESkeletonTreeFilterResult::Hidden;
-				}
+					int32 BoneMeshIndex = PreviewMeshComponent->GetBoneIndex(BoneItem->GetRowItemName());
 
-				// Remove non-vertex-weighted bones if we're filtering
-				if (BoneFilter == EBoneFilter::Weighted && !BoneItem->IsBoneWeighted(BoneMeshIndex, PreviewMeshComponent))
-				{
-					Result = ESkeletonTreeFilterResult::Hidden;
-				}
+					// Remove non-mesh bones if we're filtering
+					if ((BoneFilter == EBoneFilter::Mesh || BoneFilter == EBoneFilter::Weighted || BoneFilter == EBoneFilter::LOD) &&
+						BoneMeshIndex == INDEX_NONE)
+					{
+						Result = ESkeletonTreeFilterResult::Hidden;
+					}
 
-				// Remove non-vertex-weighted bones if we're filtering
-				if (BoneFilter == EBoneFilter::LOD && !BoneItem->IsBoneRequired(BoneMeshIndex, PreviewMeshComponent))
-				{
-					Result = ESkeletonTreeFilterResult::Hidden;
+					// Remove non-vertex-weighted bones if we're filtering
+					if (BoneFilter == EBoneFilter::Weighted && !BoneItem->IsBoneWeighted(BoneMeshIndex, PreviewMeshComponent))
+					{
+						Result = ESkeletonTreeFilterResult::Hidden;
+					}
+
+					// Remove non-vertex-weighted bones if we're filtering
+					if (BoneFilter == EBoneFilter::LOD && !BoneItem->IsBoneRequired(BoneMeshIndex, PreviewMeshComponent))
+					{
+						Result = ESkeletonTreeFilterResult::Hidden;
+					}
 				}
 			}
 		}
-	}
-	else if (InItem->IsOfType<FSkeletonTreeSocketItem>())
-	{
-		TSharedPtr<FSkeletonTreeSocketItem> SocketItem = StaticCastSharedPtr<FSkeletonTreeSocketItem>(InItem);
-
-		if (SocketFilter == ESocketFilter::None)
+		else if (InItem->IsOfType<FSkeletonTreeSocketItem>())
 		{
-			Result = ESkeletonTreeFilterResult::Hidden;
-		}
+			TSharedPtr<FSkeletonTreeSocketItem> SocketItem = StaticCastSharedPtr<FSkeletonTreeSocketItem>(InItem);
 
-		// Remove non-mesh sockets if we're filtering
-		if ((SocketFilter == ESocketFilter::Mesh || SocketFilter == ESocketFilter::None) && SocketItem->GetParentType() == ESocketParentType::Skeleton)
-		{
-			Result = ESkeletonTreeFilterResult::Hidden;
-		}
+			if (SocketFilter == ESocketFilter::None)
+			{
+				Result = ESkeletonTreeFilterResult::Hidden;
+			}
 
-		// Remove non-skeleton sockets if we're filtering
-		if ((SocketFilter == ESocketFilter::Skeleton || SocketFilter == ESocketFilter::None) && SocketItem->GetParentType() == ESocketParentType::Mesh)
-		{
-			Result = ESkeletonTreeFilterResult::Hidden;
-		}
+			// Remove non-mesh sockets if we're filtering
+			if ((SocketFilter == ESocketFilter::Mesh || SocketFilter == ESocketFilter::None) && SocketItem->GetParentType() == ESocketParentType::Skeleton)
+			{
+				Result = ESkeletonTreeFilterResult::Hidden;
+			}
 
-		if (SocketFilter == ESocketFilter::Active && SocketItem->GetParentType() == ESocketParentType::Skeleton && SocketItem->IsSocketCustomized())
-		{
-			// Don't add the skeleton socket if it's already added for the mesh
-			Result = ESkeletonTreeFilterResult::Hidden;
+			// Remove non-skeleton sockets if we're filtering
+			if ((SocketFilter == ESocketFilter::Skeleton || SocketFilter == ESocketFilter::None) && SocketItem->GetParentType() == ESocketParentType::Mesh)
+			{
+				Result = ESkeletonTreeFilterResult::Hidden;
+			}
+
+			if (SocketFilter == ESocketFilter::Active && SocketItem->GetParentType() == ESocketParentType::Skeleton && SocketItem->IsSocketCustomized())
+			{
+				// Don't add the skeleton socket if it's already added for the mesh
+				Result = ESkeletonTreeFilterResult::Hidden;
+			}
 		}
 	}
 
diff --git a/Engine/Source/Editor/SkeletonEditor/Private/SSkeletonTree.h b/Engine/Source/Editor/SkeletonEditor/Private/SSkeletonTree.h
--- a/Engine/Source/Editor/SkeletonEditor/Private/SSkeletonTree.h
+++ b/Engine/Source/Editor/SkeletonEditor/Private/SSkeletonTree.h
@@ -22,6 +22,8 @@
 #include "ISkeletonTreeBuilder.h"
 #include "ISkeletonTreeItem.h"
 #include "SkeletonTreeBuilder.h"
+#include "SSearchBox.h"
+#include "EditorUndoClient.h"
 
 class FMenuBuilder;
 class FSkeletonTreeAttachedAssetItem;
@@ -39,7 +41,8 @@ struct FNotificationInfo;
 //////////////////////////////////////////////////////////////////////////
 // SSkeletonTree
 
-class SSkeletonTree : public ISkeletonTree
+PRAGMA_DISABLE_DEPRECATION_WARNINGS
+class SSkeletonTree : public ISkeletonTree, public FEditorUndoClient
 {
 public:
 	SLATE_BEGIN_ARGS( SSkeletonTree )
@@ -63,13 +66,18 @@ public:
 	void Construct(const FArguments& InArgs, const TSharedRef<FEditableSkeleton>& InEditableSkeleton, const FSkeletonTreeArgs& InSkeletonTreeArgs);
 
 	/** ISkeletonTree interface */
+	virtual void Refresh() override;
+	virtual void RefreshFilter() override;
 	virtual TSharedRef<class IEditableSkeleton> GetEditableSkeleton() const override { return EditableSkeleton.Pin().ToSharedRef(); }
 	virtual TSharedPtr<class IPersonaPreviewScene> GetPreviewScene() const override { return PreviewScene.Pin(); }
 	virtual void SetSkeletalMesh(USkeletalMesh* NewSkeletalMesh) override;
 	virtual void SetSelectedSocket(const struct FSelectedSocketInfo& InSocketInfo) override;
 	virtual void SetSelectedBone(const FName& InBoneName) override;
 	virtual void DeselectAll() override;
+	virtual TArray<TSharedPtr<ISkeletonTreeItem>> GetSelectedItems() const override { return SkeletonTreeView->GetSelectedItems(); }
+	virtual void SelectItemsBy(TFunctionRef<bool(const TSharedRef<ISkeletonTreeItem>&, bool&)> Predicate) const override;
 	virtual void DuplicateAndSelectSocket(const FSelectedSocketInfo& SocketInfoToDuplicate, const FName& NewParentBoneName = FName()) override;
+
 	virtual void RegisterOnObjectSelected(const FOnObjectSelected& Delegate) override
 	{
 		OnObjectSelectedMulticast.Add(Delegate);
@@ -80,8 +88,23 @@ public:
 		OnObjectSelectedMulticast.RemoveAll(Widget);
 	}
 
+	virtual FDelegateHandle RegisterOnSelectionChanged(const FOnSkeletonTreeSelectionChanged& Delegate) override
+	{
+		return OnSelectionChangedMulticast.Add(Delegate);
+	}
+
+	virtual void UnregisterOnSelectionChanged(FDelegateHandle DelegateHandle) override
+	{
+		OnSelectionChangedMulticast.Remove(DelegateHandle);
+	}
+
 	virtual UBlendProfile* GetSelectedBlendProfile() override;
 	virtual void AttachAssets(const TSharedRef<ISkeletonTreeItem>& TargetItem, const TArray<FAssetData>& AssetData) override;
+	virtual TSharedPtr<SWidget> GetSearchWidget() const override { return NameFilterBox; }
+
+	/** FEditorUndoClient interface */
+	virtual void PostUndo(bool bSuccess) override;
+	virtual void PostRedo(bool bSuccess) override;
 
 	/** Creates the tree control and then populates */
 	void CreateTreeColumns();
@@ -92,8 +115,8 @@ public:
 	/** Apply filtering to the tree */
 	void ApplyFilter();
 
-	/** This triggers a rebuild of the tree after undo to make the UI consistent with the real data */
-	void PostUndo();
+	/** Set the initial expansion state of the tree items */
+	void SetInitialExpansionState();
 
 	/** Utility function to print notifications to the user */
 	void NotifyUser( FNotificationInfo& NotificationInfo );
@@ -234,6 +257,9 @@ private:
 	/** Override OnKeyDown */
 	virtual FReply OnKeyDown( const FGeometry& MyGeometry, const FKeyEvent& InKeyEvent );
 
+	/** Check whether we can delete all the selected sockets/assets */
+	bool CanDeleteSelectedRows() const;
+
 	/** Function to delete all the selected sockets/assets */
 	void OnDeleteSelectedRows();
 
@@ -262,7 +288,7 @@ private:
 	void HandleFocusCamera();
 
 	/** Handle filtering the tree  */
-	ESkeletonTreeFilterResult HandleFilterSkeletonTreeItem(const TSharedPtr<class ISkeletonTreeItem>& InItem);
+	ESkeletonTreeFilterResult HandleFilterSkeletonTreeItem(const FSkeletonTreeFilterArgs& InArgs, const TSharedPtr<class ISkeletonTreeItem>& InItem);
 
 	// Called when bone tree queries reference skeleton
 	const FReferenceSkeleton& OnGetReferenceSkeleton() const
@@ -320,7 +346,10 @@ private:
 	/** Last Cached Preview Mesh Component LOD */
 	int32 LastCachedLODForPreviewMeshComponent;
 
-	/** Delegate for when a socket is selected by clicking its hit point */
+	/** Delegate for when an item is selected */
+	FOnSkeletonTreeSelectionChangedMulticast OnSelectionChangedMulticast;
+
+	DEPRECATED(4.17, "Please use OnSelectionChangedMulticast")
 	FOnObjectSelectedMulticast OnObjectSelectedMulticast;
 
 	/** Selection recursion guard flags */
@@ -340,8 +369,24 @@ private:
 	/** Compiled filter search terms. */
 	TSharedPtr<class FTextFilterExpressionEvaluator> TextFilterPtr;
 
-	/** Proxy object used to display and edit bone transforms in details panels */
+	/** Proxy object used to display and edit bone transforms in details panels. Note this is only kept for backwards compatibility (used with OnObjectSelectedMulticast) */
 	class UBoneProxy* BoneProxy;
 
+	/** Whether to allow operations that modify the mesh */
+	bool bAllowMeshOperations;
+
+	/** Whether to allow operations that modify the mesh */
+	bool bAllowSkeletonOperations;
+
+	/** Extenders for menus */
+	TSharedPtr<FExtender> Extenders;
+
+	/** Delegate that allows custom filtering text to be shown on the filter button */
+	FOnGetFilterText OnGetFilterText;
+
+	/** The mode that this skeleton tree is in */
+	ESkeletonTreeMode Mode;
+
 	friend struct FScopedSavedSelection;
 }; 
+PRAGMA_ENABLE_DEPRECATION_WARNINGS
\ No newline at end of file
diff --git a/Engine/Source/Editor/SkeletonEditor/Private/SSkeletonTreeRow.cpp b/Engine/Source/Editor/SkeletonEditor/Private/SSkeletonTreeRow.cpp
--- a/Engine/Source/Editor/SkeletonEditor/Private/SSkeletonTreeRow.cpp
+++ b/Engine/Source/Editor/SkeletonEditor/Private/SSkeletonTreeRow.cpp
@@ -2,6 +2,7 @@
 
 #include "SSkeletonTreeRow.h"
 #include "ISkeletonTree.h"
+#include "Preferences/PersonaOptions.h"
 
 void SSkeletonTreeRow::Construct( const FArguments& InArgs, const TSharedRef<STableViewBase>& InOwnerTableView )
 {
@@ -44,7 +45,11 @@ TSharedRef< SWidget > SSkeletonTreeRow::GenerateWidgetForColumn( const FName& Co
 	{
 		TSharedPtr< SHorizontalBox > RowBox;
 
-		SAssignNew( RowBox, SHorizontalBox );
+		SAssignNew( RowBox, SHorizontalBox )
+			.Visibility_Lambda([this]()
+			{
+				return Item.Pin()->GetFilterResult() == ESkeletonTreeFilterResult::ShownDescendant && GetMutableDefault<UPersonaOptions>()->bHideParentsWhenFiltering ? EVisibility::Collapsed : EVisibility::Visible;
+			});
 
 		RowBox->AddSlot()
 			.AutoWidth()
diff --git a/Engine/Source/Editor/SkeletonEditor/Private/SkeletonEditor.cpp b/Engine/Source/Editor/SkeletonEditor/Private/SkeletonEditor.cpp
--- a/Engine/Source/Editor/SkeletonEditor/Private/SkeletonEditor.cpp
+++ b/Engine/Source/Editor/SkeletonEditor/Private/SkeletonEditor.cpp
@@ -18,6 +18,8 @@
 #include "IAssetFamily.h"
 #include "PersonaCommonCommands.h"
 #include "IEditableSkeleton.h"
+#include "ISkeletonTreeItem.h"
+#include "Algo/Transform.h"
 
 const FName SkeletonEditorAppIdentifier = FName(TEXT("SkeletonEditorApp"));
 
@@ -78,8 +80,8 @@ void FSkeletonEditor::InitSkeletonEditor(const EToolkitMode::Type Mode, const TS
 	TSharedRef<IAssetFamily> AssetFamily = PersonaModule.CreatePersonaAssetFamily(InSkeleton);
 	AssetFamily->RecordAssetOpened(FAssetData(InSkeleton));
 
-	FSkeletonTreeArgs SkeletonTreeArgs(OnPostUndo);
-	SkeletonTreeArgs.OnObjectSelected = FOnObjectSelected::CreateSP(this, &FSkeletonEditor::HandleObjectSelected);
+	FSkeletonTreeArgs SkeletonTreeArgs;
+	SkeletonTreeArgs.OnSelectionChanged = FOnSkeletonTreeSelectionChanged::CreateSP(this, &FSkeletonEditor::HandleSelectionChanged);
 	SkeletonTreeArgs.PreviewScene = PersonaToolkit->GetPreviewScene();
 
 	ISkeletonEditorModule& SkeletonEditorModule = FModuleManager::GetModuleChecked<ISkeletonEditorModule>("SkeletonEditor");
@@ -247,6 +249,16 @@ void FSkeletonEditor::HandleObjectSelected(UObject* InObject)
 	}
 }
 
+void FSkeletonEditor::HandleSelectionChanged(const TArrayView<TSharedPtr<ISkeletonTreeItem>>& InSelectedItems, ESelectInfo::Type InSelectInfo)
+{
+	if (DetailsView.IsValid())
+	{
+		TArray<UObject*> Objects;
+		Algo::TransformIf(InSelectedItems, Objects, [](const TSharedPtr<ISkeletonTreeItem>& InItem) { return InItem->GetObject() != nullptr; }, [](const TSharedPtr<ISkeletonTreeItem>& InItem) { return InItem->GetObject(); });
+		DetailsView->SetObjects(Objects);
+	}
+}
+
 void FSkeletonEditor::PostUndo(bool bSuccess)
 {
 	OnPostUndo.Broadcast();
diff --git a/Engine/Source/Editor/SkeletonEditor/Private/SkeletonEditor.h b/Engine/Source/Editor/SkeletonEditor/Private/SkeletonEditor.h
--- a/Engine/Source/Editor/SkeletonEditor/Private/SkeletonEditor.h
+++ b/Engine/Source/Editor/SkeletonEditor/Private/SkeletonEditor.h
@@ -9,11 +9,13 @@
 #include "EditorUndoClient.h"
 #include "IDetailsView.h"
 #include "ISkeletonEditor.h"
+#include "ArrayView.h"
 
 class IPersonaToolkit;
 class IPersonaViewport;
 class ISkeletonTree;
 class USkeleton;
+class ISkeletonTreeItem;
 
 namespace SkeletonEditorModes
 {
@@ -78,6 +80,8 @@ public:
 
 	void HandleObjectSelected(UObject* InObject);
 
+	void HandleSelectionChanged(const TArrayView<TSharedPtr<ISkeletonTreeItem>>& InSelectedItems, ESelectInfo::Type InSelectInfo);
+
 	void HandleDetailsCreated(const TSharedRef<class IDetailsView>& InDetailsView);
 
 private:
diff --git a/Engine/Source/Editor/SkeletonEditor/Private/SkeletonEditorMode.cpp b/Engine/Source/Editor/SkeletonEditor/Private/SkeletonEditorMode.cpp
--- a/Engine/Source/Editor/SkeletonEditor/Private/SkeletonEditorMode.cpp
+++ b/Engine/Source/Editor/SkeletonEditor/Private/SkeletonEditorMode.cpp
@@ -22,7 +22,11 @@ FSkeletonEditorMode::FSkeletonEditorMode(TSharedRef<FWorkflowCentricApplication>
 	FOnObjectSelected OnObjectSelected = FOnObjectSelected::CreateSP(&SkeletonEditor.Get(), &FSkeletonEditor::HandleObjectSelected);
 	FPersonaModule& PersonaModule = FModuleManager::LoadModuleChecked<FPersonaModule>("Persona");
 	TabFactories.RegisterFactory(PersonaModule.CreateDetailsTabFactory(InHostingApp, FOnDetailsCreated::CreateSP(&SkeletonEditor.Get(), &FSkeletonEditor::HandleDetailsCreated)));
-	TabFactories.RegisterFactory(PersonaModule.CreatePersonaViewportTabFactory(InHostingApp, InSkeletonTree, SkeletonEditor->GetPersonaToolkit()->GetPreviewScene(), SkeletonEditor->OnPostUndo, nullptr, FOnViewportCreated(), true, true));
+
+	FPersonaViewportArgs ViewportArgs(InSkeletonTree, SkeletonEditor->GetPersonaToolkit()->GetPreviewScene(), SkeletonEditor->OnPostUndo);
+
+	TabFactories.RegisterFactory(PersonaModule.CreatePersonaViewportTabFactory(InHostingApp, ViewportArgs));
+
 	TabFactories.RegisterFactory(PersonaModule.CreateAnimNotifiesTabFactory(InHostingApp, InSkeletonTree->GetEditableSkeleton(), SkeletonEditor->OnChangeAnimNotifies, SkeletonEditor->OnPostUndo, OnObjectsSelected));
 	TabFactories.RegisterFactory(PersonaModule.CreateAdvancedPreviewSceneTabFactory(InHostingApp, SkeletonEditor->GetPersonaToolkit()->GetPreviewScene()));
 	TabFactories.RegisterFactory(PersonaModule.CreateRetargetManagerTabFactory(InHostingApp, InSkeletonTree->GetEditableSkeleton(), SkeletonEditor->GetPersonaToolkit()->GetPreviewScene(), SkeletonEditor->OnPostUndo));
diff --git a/Engine/Source/Editor/SkeletonEditor/Private/SkeletonEditorModule.cpp b/Engine/Source/Editor/SkeletonEditor/Private/SkeletonEditorModule.cpp
--- a/Engine/Source/Editor/SkeletonEditor/Private/SkeletonEditorModule.cpp
+++ b/Engine/Source/Editor/SkeletonEditor/Private/SkeletonEditorModule.cpp
@@ -64,6 +64,11 @@ public:
 		return FSkeletonTreeManager::Get().CreateSkeletonTree(InSkeleton, InSkeletonTreeArgs);
 	}
 
+	virtual TSharedRef<ISkeletonTree> CreateSkeletonTree(const TSharedRef<IEditableSkeleton>& InEditableSkeleton, const FSkeletonTreeArgs& InSkeletonTreeArgs) override
+	{
+		return FSkeletonTreeManager::Get().CreateSkeletonTree(InEditableSkeleton, InSkeletonTreeArgs);
+	}
+
 	virtual TSharedRef<class FWorkflowTabFactory> CreateSkeletonTreeTabFactory(const TSharedRef<class FWorkflowCentricApplication>& InHostingApp, const TSharedRef<ISkeletonTree>& InSkeletonTree) override
 	{
 		return MakeShareable(new FSkeletonTreeSummoner(InHostingApp, InSkeletonTree));
diff --git a/Engine/Source/Editor/SkeletonEditor/Private/SkeletonTreeBoneItem.cpp b/Engine/Source/Editor/SkeletonEditor/Private/SkeletonTreeBoneItem.cpp
--- a/Engine/Source/Editor/SkeletonEditor/Private/SkeletonTreeBoneItem.cpp
+++ b/Engine/Source/Editor/SkeletonEditor/Private/SkeletonTreeBoneItem.cpp
@@ -18,16 +18,36 @@
 #include "Framework/Commands/UIAction.h"
 #include "Framework/MultiBox/MultiBoxBuilder.h"
 #include "SkeletalMeshTypes.h"
+#include "UObject/Package.h"
 
 #define LOCTEXT_NAMESPACE "FSkeletonTreeBoneItem"
 
-TSharedRef<ITableRow> FSkeletonTreeBoneItem::MakeTreeRowWidget(const TSharedRef<STableViewBase>& OwnerTable, const TAttribute<FText>& FilterText )
+FSkeletonTreeBoneItem::FSkeletonTreeBoneItem(const FName& InBoneName, const TSharedRef<class ISkeletonTree>& InSkeletonTree)
+	: FSkeletonTreeItem(InSkeletonTree)
+	, BoneName(InBoneName)
+	, bWeightedBone(false)
+	, bRequiredBone(false)
 {
-	return
-		SNew( SSkeletonTreeRow, OwnerTable )
-		.Item( SharedThis(this) )
-		.FilterText( FilterText )
-		.OnDraggingItem( this, &FSkeletonTreeBoneItem::OnDragDetected );
+	static const FString BoneProxyPrefix(TEXT("BONEPROXY_"));
+
+	BoneProxy = NewObject<UBoneProxy>(GetTransientPackage(), *(BoneProxyPrefix + FString::Printf(TEXT("%p"), &InSkeletonTree.Get()) + InBoneName.ToString()));
+	BoneProxy->SetFlags(RF_Transactional);
+	BoneProxy->AddToRoot();
+	BoneProxy->BoneName = InBoneName;
+	TSharedPtr<IPersonaPreviewScene> PreviewScene = InSkeletonTree->GetPreviewScene();
+	if (PreviewScene.IsValid())
+	{
+		BoneProxy->SkelMeshComponent = PreviewScene->GetPreviewMeshComponent();
+	}
+}
+
+FSkeletonTreeBoneItem::~FSkeletonTreeBoneItem()
+{
+	if (BoneProxy)
+	{
+		BoneProxy->RemoveFromRoot();
+		BoneProxy = nullptr;
+	}
 }
 
 EVisibility FSkeletonTreeBoneItem::GetLODIconVisibility() const
@@ -236,6 +256,11 @@ void FSkeletonTreeBoneItem::CacheLODChange(UDebugSkelMeshComponent* PreviewCompo
 	}
 }
 
+void FSkeletonTreeBoneItem::EnableBoneProxyTick(bool bEnable)
+{
+	BoneProxy->bIsTickable = bEnable;
+}
+
 FSlateColor FSkeletonTreeBoneItem::GetBoneTextColor() const
 {
 	if (FilterResult == ESkeletonTreeFilterResult::ShownDescendant)
diff --git a/Engine/Source/Editor/SkeletonEditor/Private/SkeletonTreeBoneItem.h b/Engine/Source/Editor/SkeletonEditor/Private/SkeletonTreeBoneItem.h
--- a/Engine/Source/Editor/SkeletonEditor/Private/SkeletonTreeBoneItem.h
+++ b/Engine/Source/Editor/SkeletonEditor/Private/SkeletonTreeBoneItem.h
@@ -14,6 +14,7 @@
 #include "SkeletonTreeItem.h"
 #include "Animation/Skeleton.h"
 #include "Widgets/Input/SComboButton.h"
+#include "BoneProxy.h"
 
 class UDebugSkelMeshComponent;
 
@@ -22,23 +23,18 @@ class FSkeletonTreeBoneItem : public FSkeletonTreeItem
 public:
 	SKELETON_TREE_ITEM_TYPE(FSkeletonTreeBoneItem, FSkeletonTreeItem)
 
-	FSkeletonTreeBoneItem(const FName& InBoneName, const TSharedRef<class ISkeletonTree>& InSkeletonTree)
-		: FSkeletonTreeItem(InSkeletonTree)
-		, BoneName(InBoneName)
-		, bWeightedBone(false)
-		, bRequiredBone(false)
-	{}
-	virtual ~FSkeletonTreeBoneItem()
-	{}
+	FSkeletonTreeBoneItem(const FName& InBoneName, const TSharedRef<class ISkeletonTree>& InSkeletonTree);
+	virtual ~FSkeletonTreeBoneItem();
 
 	/** ISkeletonTreeItem interface */
-	virtual TSharedRef<ITableRow> MakeTreeRowWidget(const TSharedRef<STableViewBase>& OwnerTable, const TAttribute<FText>& FilterText) override;
 	virtual void GenerateWidgetForNameColumn(TSharedPtr< SHorizontalBox > Box, const TAttribute<FText>& FilterText, FIsSelected InIsSelected) override;
 	virtual TSharedRef< SWidget > GenerateWidgetForDataColumn(const FName& DataColumnName) override;
 	virtual FName GetRowItemName() const override { return BoneName; }
 	virtual void HandleDragEnter(const FDragDropEvent& DragDropEvent) override;
 	virtual void HandleDragLeave(const FDragDropEvent& DragDropEvent) override;
 	virtual FReply HandleDrop(const FDragDropEvent& DragDropEvent) override;
+	virtual FReply OnDragDetected(const FGeometry& MyGeometry, const FPointerEvent& MouseEvent) override;
+	virtual UObject* GetObject() const override { return BoneProxy; }
 
 	/** Check to see if the specified bone is weighted in the specified component */
 	bool IsBoneWeighted(int32 MeshBoneIndex, UDebugSkelMeshComponent* PreviewComponent);
@@ -49,6 +45,9 @@ public:
 	/** Cache to use Bold Font or not */
 	void CacheLODChange(UDebugSkelMeshComponent* PreviewComponent);
 
+	/** Enable and disable the bone proxy ticking */
+	void EnableBoneProxyTick(bool bEnable);
+
 private:
 	/** Gets the font for displaying bone text in the skeletal tree */
 	FSlateFontInfo GetBoneTextFont() const;
@@ -74,9 +73,6 @@ private:
 	/** Callback from a slider widget if the text entry or slider movement is used */
 	void OnBlendSliderCommitted(float NewValue, ETextCommit::Type CommitType);
 
-	/** Handle dragging a bone */
-	FReply OnDragDetected(const FGeometry& MyGeometry, const FPointerEvent& MouseEvent);
-
 	/** Set Translation Retargeting Mode for this bone. */
 	void SetBoneTranslationRetargetingMode(EBoneTranslationRetargetingMode::Type NewRetargetingMode);
 
@@ -84,6 +80,9 @@ private:
 	void SetBoneBlendProfileScale(float NewScale, bool bRecurse);
 
 private:
+	/** Bone proxy object */
+	UBoneProxy* BoneProxy;
+
 	/** The actual bone data that we create Slate widgets to display */
 	FName BoneName;
 
diff --git a/Engine/Source/Editor/SkeletonEditor/Private/SkeletonTreeBuilder.cpp b/Engine/Source/Editor/SkeletonEditor/Private/SkeletonTreeBuilder.cpp
--- a/Engine/Source/Editor/SkeletonEditor/Private/SkeletonTreeBuilder.cpp
+++ b/Engine/Source/Editor/SkeletonEditor/Private/SkeletonTreeBuilder.cpp
@@ -8,21 +8,39 @@
 #include "SkeletonTreeVirtualBoneItem.h"
 #include "Animation/DebugSkelMeshComponent.h"
 
-void FSkeletonTreeBuilderOutput::Add(const TSharedPtr<class ISkeletonTreeItem>& InItem, const FName& InParentName, const FName& InParentType)
+#define LOCTEXT_NAMESPACE "SkeletonTreeBuilder"
+
+void FSkeletonTreeBuilderOutput::Add(const TSharedPtr<class ISkeletonTreeItem>& InItem, const FName& InParentName, const FName& InParentType, bool bAddToHead)
 {
-	Add(InItem, InParentName, TArray<FName, TInlineAllocator<1>>({ InParentType }));
+	Add(InItem, InParentName, TArray<FName, TInlineAllocator<1>>({ InParentType }), bAddToHead);
 }
 
-void FSkeletonTreeBuilderOutput::Add(const TSharedPtr<class ISkeletonTreeItem>& InItem, const FName& InParentName, TArrayView<const FName> InParentTypes)
+void FSkeletonTreeBuilderOutput::Add(const TSharedPtr<class ISkeletonTreeItem>& InItem, const FName& InParentName, TArrayView<const FName> InParentTypes, bool bAddToHead)
 {
 	TSharedPtr<ISkeletonTreeItem> ParentItem = Find(InParentName, InParentTypes);
 	if (ParentItem.IsValid())
 	{
-		ParentItem->GetChildren().Add(InItem);
+		InItem->SetParent(ParentItem);
+
+		if (bAddToHead)
+		{
+			ParentItem->GetChildren().Insert(InItem, 0);
+		}
+		else
+		{
+			ParentItem->GetChildren().Add(InItem);
+		}
 	}
 	else
 	{
-		Items.Add(InItem);
+		if (bAddToHead)
+		{
+			Items.Insert(InItem, 0);
+		}
+		else
+		{
+			Items.Add(InItem);
+		}
 	}
 
 	LinearItems.Add(InItem);
@@ -56,18 +74,25 @@ TSharedPtr<class ISkeletonTreeItem> FSkeletonTreeBuilderOutput::Find(const FName
 	return nullptr;
 }
 
-FSkeletonTreeBuilder::FSkeletonTreeBuilder(const FSkeletonTreeBuilderArgs& InBuilderArgs, FOnFilterSkeletonTreeItem InOnFilterSkeletonTreeItem, const TSharedRef<class ISkeletonTree>& InSkeletonTree, const TSharedPtr<class IPersonaPreviewScene>& InPreviewScene)
+FSkeletonTreeBuilder::FSkeletonTreeBuilder(const FSkeletonTreeBuilderArgs& InBuilderArgs)
 	: BuilderArgs(InBuilderArgs)
-	, OnFilterSkeletonTreeItem(InOnFilterSkeletonTreeItem)
-	, SkeletonTreePtr(InSkeletonTree)
-	, EditableSkeletonPtr(InSkeletonTree->GetEditableSkeleton())
-	, PreviewScenePtr(InPreviewScene)
 {
 }
 
+void FSkeletonTreeBuilder::Initialize(const TSharedRef<class ISkeletonTree>& InSkeletonTree, const TSharedPtr<class IPersonaPreviewScene>& InPreviewScene, FOnFilterSkeletonTreeItem InOnFilterSkeletonTreeItem)
+{
+	SkeletonTreePtr = InSkeletonTree;
+	EditableSkeletonPtr = InSkeletonTree->GetEditableSkeleton();
+	PreviewScenePtr = InPreviewScene;
+	OnFilterSkeletonTreeItem = InOnFilterSkeletonTreeItem;
+}
+
 void FSkeletonTreeBuilder::Build(FSkeletonTreeBuilderOutput& Output)
 {
-	AddBones(Output);
+	if(BuilderArgs.bShowBones)
+	{
+		AddBones(Output);
+	}
 
 	if (BuilderArgs.bShowSockets)
 	{
@@ -91,7 +116,7 @@ void FSkeletonTreeBuilder::Filter(const FSkeletonTreeFilterArgs& InArgs, const T
 
 	for (const TSharedPtr<ISkeletonTreeItem>& Item : InItems)
 	{
-		if (InArgs.bWillFilter && InArgs.bFlattenHierarchyOnFilter)
+		if (InArgs.TextFilter.IsValid() && InArgs.bFlattenHierarchyOnFilter)
 		{
 			FilterRecursive(InArgs, Item, OutFilteredItems);
 		}
@@ -114,9 +139,9 @@ ESkeletonTreeFilterResult FSkeletonTreeBuilder::FilterRecursive(const FSkeletonT
 
 	InItem->GetFilteredChildren().Empty();
 
-	if (InArgs.bWillFilter && InArgs.bFlattenHierarchyOnFilter)
+	if (InArgs.TextFilter.IsValid() && InArgs.bFlattenHierarchyOnFilter)
 	{
-		FilterResult = OnFilterSkeletonTreeItem.Execute(InItem);
+		FilterResult = FilterItem(InArgs, InItem);
 		InItem->SetFilterResult(FilterResult);
 
 		if (FilterResult != ESkeletonTreeFilterResult::Hidden)
@@ -146,7 +171,7 @@ ESkeletonTreeFilterResult FSkeletonTreeBuilder::FilterRecursive(const FSkeletonT
 			}
 		}
 
-		FilterResult = OnFilterSkeletonTreeItem.Execute(InItem);
+		FilterResult = FilterItem(InArgs, InItem);
 		if (DescendantsFilterResult > FilterResult)
 		{
 			FilterResult = ESkeletonTreeFilterResult::ShownDescendant;
@@ -158,6 +183,26 @@ ESkeletonTreeFilterResult FSkeletonTreeBuilder::FilterRecursive(const FSkeletonT
 	return FilterResult;
 }
 
+ESkeletonTreeFilterResult FSkeletonTreeBuilder::FilterItem(const FSkeletonTreeFilterArgs& InArgs, const TSharedPtr<class ISkeletonTreeItem>& InItem)
+{
+	return OnFilterSkeletonTreeItem.Execute(InArgs, InItem);
+}
+
+bool FSkeletonTreeBuilder::IsShowingBones() const
+{
+	return BuilderArgs.bShowBones;
+}
+
+bool FSkeletonTreeBuilder::IsShowingSockets() const
+{
+	return BuilderArgs.bShowSockets;
+}
+
+bool FSkeletonTreeBuilder::IsShowingAttachedAssets() const
+{
+	return BuilderArgs.bShowAttachedAssets;
+}
+
 void FSkeletonTreeBuilder::AddBones(FSkeletonTreeBuilderOutput& Output)
 {
 	const USkeleton& Skeleton = EditableSkeletonPtr.Pin()->GetSkeleton();
@@ -175,9 +220,6 @@ void FSkeletonTreeBuilder::AddBones(FSkeletonTreeBuilderOutput& Output)
 		}
 		TSharedRef<ISkeletonTreeItem> DisplayBone = CreateBoneTreeItem(BoneName);
 		Output.Add(DisplayBone, ParentName, FSkeletonTreeBoneItem::GetTypeId());
-
-		// @TODO: Persona2.0: Item expansion should be saved
-		//SkeletonTreeView->SetItemExpansion(DisplayBone, true);
 	}
 }
 
@@ -244,9 +286,6 @@ void FSkeletonTreeBuilder::AddSocketsFromData(const TArray< USkeletalMeshSocket*
 		}
 
 		Output.Add(CreateSocketTreeItem(Socket, ParentType, bIsCustomized), Socket->BoneName, FSkeletonTreeBoneItem::GetTypeId());
-
-		// @TODO: Persona2.0: Item expansion should be saved
-	//	SkeletonTreeView->SetItemExpansion(DisplaySocket, true);
 	}
 }
 
@@ -288,3 +327,5 @@ TSharedRef<class ISkeletonTreeItem> FSkeletonTreeBuilder::CreateVirtualBoneTreeI
 {
 	return MakeShareable(new FSkeletonTreeVirtualBoneItem(InBoneName, SkeletonTreePtr.Pin().ToSharedRef()));
 }
+
+#undef LOCTEXT_NAMESPACE
\ No newline at end of file
diff --git a/Engine/Source/Editor/SkeletonEditor/Private/SkeletonTreeCommands.cpp b/Engine/Source/Editor/SkeletonEditor/Private/SkeletonTreeCommands.cpp
--- a/Engine/Source/Editor/SkeletonEditor/Private/SkeletonTreeCommands.cpp
+++ b/Engine/Source/Editor/SkeletonEditor/Private/SkeletonTreeCommands.cpp
@@ -12,7 +12,8 @@ void FSkeletonTreeCommands::RegisterCommands()
 	UI_COMMAND( ShowWeightedBones, "Show Weighted Bones", "Show bones that have vertices weighted to them", EUserInterfaceActionType::RadioButton, FInputChord() );
 	UI_COMMAND( HideBones, "Hide Bones", "Hides all bones (sockets and attached assets will still be listed)", EUserInterfaceActionType::RadioButton, FInputChord() );
 	UI_COMMAND( ShowRetargeting, "Show Retargeting Options", "Show retargeting options for each bone.", EUserInterfaceActionType::ToggleButton, FInputChord());
-	UI_COMMAND( FilteringFlattensHierarchy, "Filtering Flattens Hierarchy", "Whether to keep the hierarchy or flatten it when searching for bones, sockets etc.", EUserInterfaceActionType::ToggleButton, FInputChord());
+	UI_COMMAND( FilteringFlattensHierarchy, "Filtering Flattens Hierarchy", "Whether to keep the hierarchy or flatten it when searching for tree items", EUserInterfaceActionType::ToggleButton, FInputChord());
+	UI_COMMAND( HideParentsWhenFiltering, "Hide Parents When Filtering", "Whether to show parent items grayed out, or hide them entirely when filtering", EUserInterfaceActionType::ToggleButton, FInputChord());
 	
 	UI_COMMAND( CopyBoneNames, "Copy Selected Bone Names", "Copy selected bone names to clipboard", EUserInterfaceActionType::Button, FInputChord() );
 	UI_COMMAND( ResetBoneTransforms, "Reset Selected Bone Transforms", "Reset the transforms of the selected bones", EUserInterfaceActionType::Button, FInputChord() );
diff --git a/Engine/Source/Editor/SkeletonEditor/Private/SkeletonTreeCommands.h b/Engine/Source/Editor/SkeletonEditor/Private/SkeletonTreeCommands.h
--- a/Engine/Source/Editor/SkeletonEditor/Private/SkeletonTreeCommands.h
+++ b/Engine/Source/Editor/SkeletonEditor/Private/SkeletonTreeCommands.h
@@ -48,6 +48,9 @@ public:
 	/** Flatten hierarchy on filter */
 	TSharedPtr< FUICommandInfo > FilteringFlattensHierarchy;
 
+	/** Hide parents on filter */
+	TSharedPtr< FUICommandInfo > HideParentsWhenFiltering;
+
 	/** Add a socket to the skeleton */
 	TSharedPtr< FUICommandInfo > AddSocket;
 
diff --git a/Engine/Source/Editor/SkeletonEditor/Private/SkeletonTreeItem.cpp b/Engine/Source/Editor/SkeletonEditor/Private/SkeletonTreeItem.cpp
--- a/Engine/Source/Editor/SkeletonEditor/Private/SkeletonTreeItem.cpp
+++ b/Engine/Source/Editor/SkeletonEditor/Private/SkeletonTreeItem.cpp
@@ -7,5 +7,6 @@ TSharedRef<ITableRow> FSkeletonTreeItem::MakeTreeRowWidget(const TSharedRef<STab
 {
 	return SNew(SSkeletonTreeRow, InOwnerTable)
 		.FilterText(InFilterText)
-		.Item(SharedThis(this));
+		.Item(SharedThis(this))
+		.OnDraggingItem(this, &FSkeletonTreeItem::OnDragDetected);
 }
diff --git a/Engine/Source/Editor/SkeletonEditor/Private/SkeletonTreeManager.cpp b/Engine/Source/Editor/SkeletonEditor/Private/SkeletonTreeManager.cpp
--- a/Engine/Source/Editor/SkeletonEditor/Private/SkeletonTreeManager.cpp
+++ b/Engine/Source/Editor/SkeletonEditor/Private/SkeletonTreeManager.cpp
@@ -10,10 +10,9 @@ FSkeletonTreeManager& FSkeletonTreeManager::Get()
 	return TheManager;
 }
 
-TSharedRef<ISkeletonTree> FSkeletonTreeManager::CreateSkeletonTree(class USkeleton* InSkeleton, const FSkeletonTreeArgs& InSkeletonTreeArgs)
+TSharedRef<ISkeletonTree> FSkeletonTreeManager::CreateSkeletonTree(const TSharedRef<IEditableSkeleton>& InEditableSkeleton, const FSkeletonTreeArgs& InSkeletonTreeArgs)
 {
-	TSharedPtr<FEditableSkeleton> EditableSkeleton = CreateEditableSkeleton(InSkeleton);
-	TSharedPtr<ISkeletonTree> SkeletonTree = EditableSkeleton->CreateSkeletonTree(InSkeletonTreeArgs);
+	TSharedPtr<ISkeletonTree> SkeletonTree = StaticCastSharedRef<FEditableSkeleton>(InEditableSkeleton)->CreateSkeletonTree(InSkeletonTreeArgs);
 
 	// compact skeletons that are no longer being edited
 	bool bRemoved = false;
@@ -36,6 +35,11 @@ TSharedRef<ISkeletonTree> FSkeletonTreeManager::CreateSkeletonTree(class USkelet
 	return SkeletonTree.ToSharedRef();
 }
 
+TSharedRef<ISkeletonTree> FSkeletonTreeManager::CreateSkeletonTree(class USkeleton* InSkeleton, const FSkeletonTreeArgs& InSkeletonTreeArgs)
+{
+	return CreateSkeletonTree(CreateEditableSkeleton(InSkeleton), InSkeletonTreeArgs);
+}
+
 TSharedRef<FEditableSkeleton> FSkeletonTreeManager::CreateEditableSkeleton(class USkeleton* InSkeleton)
 {
 	TWeakPtr<FEditableSkeleton>* EditableSkeletonPtr = EditableSkeletons.Find(InSkeleton);
diff --git a/Engine/Source/Editor/SkeletonEditor/Private/SkeletonTreeManager.h b/Engine/Source/Editor/SkeletonEditor/Private/SkeletonTreeManager.h
--- a/Engine/Source/Editor/SkeletonEditor/Private/SkeletonTreeManager.h
+++ b/Engine/Source/Editor/SkeletonEditor/Private/SkeletonTreeManager.h
@@ -4,7 +4,11 @@
 
 #include "CoreMinimal.h"
 
+class ISkeletonTree;
 class FEditableSkeleton;
+class IEditableSkeleton;
+class USkeleton;
+struct FSkeletonTreeArgs;
 
 /** Central registry of skeleton trees */
 class FSkeletonTreeManager
@@ -13,11 +17,14 @@ public:
 	/** Singleton access */
 	static FSkeletonTreeManager& Get();
 
+	/** Create a skeleton tree for the given an already-existing editable skeleton */
+	TSharedRef<ISkeletonTree> CreateSkeletonTree(const TSharedRef<IEditableSkeleton>& InEditableSkeleton, const FSkeletonTreeArgs& InSkeletonTreeArgs);
+
 	/** Create a skeleton tree for the requested skeleton */
-	TSharedRef<class ISkeletonTree> CreateSkeletonTree(class USkeleton* InSkeleton, const struct FSkeletonTreeArgs& InSkeletonTreeArgs);
+	TSharedRef<ISkeletonTree> CreateSkeletonTree(USkeleton* InSkeleton, const FSkeletonTreeArgs& InSkeletonTreeArgs);
 
 	/** Edit a USkeleton via FEditableSkeleton */
-	TSharedRef<class FEditableSkeleton> CreateEditableSkeleton(class USkeleton* InSkeleton);
+	TSharedRef<FEditableSkeleton> CreateEditableSkeleton(USkeleton* InSkeleton);
 
 private:
 	/** Hidden constructor */
@@ -25,5 +32,5 @@ private:
 
 private:
 	/** Map from skeletons to editable skeletons */
-	TMap<class USkeleton*, TWeakPtr<class FEditableSkeleton>> EditableSkeletons;
+	TMap<USkeleton*, TWeakPtr<FEditableSkeleton>> EditableSkeletons;
 };
diff --git a/Engine/Source/Editor/SkeletonEditor/Private/SkeletonTreeSocketItem.cpp b/Engine/Source/Editor/SkeletonEditor/Private/SkeletonTreeSocketItem.cpp
--- a/Engine/Source/Editor/SkeletonEditor/Private/SkeletonTreeSocketItem.cpp
+++ b/Engine/Source/Editor/SkeletonEditor/Private/SkeletonTreeSocketItem.cpp
@@ -16,15 +16,6 @@
 
 #define LOCTEXT_NAMESPACE "FSkeletonTreeSocketItem"
 
-TSharedRef<ITableRow> FSkeletonTreeSocketItem::MakeTreeRowWidget(const TSharedRef<STableViewBase>& InOwnerTable, const TAttribute<FText>& InFilterText )
-{
-	return
-		SNew( SSkeletonTreeRow, InOwnerTable )
-		.Item( SharedThis(this) )
-		.FilterText( InFilterText )
-		.OnDraggingItem( this, &FSkeletonTreeSocketItem::OnDragDetected );
-}
-
 void FSkeletonTreeSocketItem::GenerateWidgetForNameColumn( TSharedPtr< SHorizontalBox > Box, const TAttribute<FText>& FilterText, FIsSelected InIsSelected )
 {
 	const FSlateBrush* SocketIcon = ( ParentType == ESocketParentType::Mesh ) ?
diff --git a/Engine/Source/Editor/SkeletonEditor/Private/SkeletonTreeSocketItem.h b/Engine/Source/Editor/SkeletonEditor/Private/SkeletonTreeSocketItem.h
--- a/Engine/Source/Editor/SkeletonEditor/Private/SkeletonTreeSocketItem.h
+++ b/Engine/Source/Editor/SkeletonEditor/Private/SkeletonTreeSocketItem.h
@@ -31,7 +31,6 @@ public:
 	{}
 
 	/** ISkeletonTreeItem interface */
-	virtual TSharedRef<ITableRow> MakeTreeRowWidget(const TSharedRef<STableViewBase>& InOwnerTable, const TAttribute<FText>& InFilterText) override;
 	virtual void GenerateWidgetForNameColumn(TSharedPtr< SHorizontalBox > Box, const TAttribute<FText>& FilterText, FIsSelected InIsSelected) override;
 	virtual TSharedRef< SWidget > GenerateWidgetForDataColumn(const FName& DataColumnName) override;
 	virtual TSharedRef< SWidget > GenerateInlineEditWidget(const TAttribute<FText>& FilterText, FIsSelected InIsSelected) override;
@@ -41,6 +40,8 @@ public:
 	virtual FName GetRowItemName() const override { return Socket->SocketName; }
 	virtual void RequestRename() override;
 	virtual void OnItemDoubleClicked() override;
+	virtual FReply OnDragDetected(const FGeometry& MyGeometry, const FPointerEvent& MouseEvent) override;
+	virtual UObject* GetObject() const override { return Socket; }
 
 	/** Get the parent type */
 	ESocketParentType GetParentType() const { return ParentType; }
@@ -65,10 +66,7 @@ private:
 	FText GetSocketToolTip();
 
 	/** Get the color to display for the socket name */
-	FSlateColor GetTextColor() const;
-
-	/** Handle dragging a socket */
-	FReply OnDragDetected(const FGeometry& MyGeometry, const FPointerEvent& MouseEvent);
+	FSlateColor GetTextColor() const;	
 
 	/** Return socket name as FText for display in skeleton tree */
 	FText GetSocketNameAsText() const { return FText::FromName(Socket->SocketName); }
diff --git a/Engine/Source/Editor/SkeletonEditor/Private/SkeletonTreeVirtualBoneItem.cpp b/Engine/Source/Editor/SkeletonEditor/Private/SkeletonTreeVirtualBoneItem.cpp
--- a/Engine/Source/Editor/SkeletonEditor/Private/SkeletonTreeVirtualBoneItem.cpp
+++ b/Engine/Source/Editor/SkeletonEditor/Private/SkeletonTreeVirtualBoneItem.cpp
@@ -8,15 +8,35 @@
 #include "Widgets/Images/SImage.h"
 #include "SInlineEditableTextBlock.h"
 #include "UnrealString.h"
+#include "Animation/DebugSkelMeshComponent.h"
+#include "UObject/Package.h"
 
 #define LOCTEXT_NAMESPACE "FSkeletonTreeVirtualBoneItem"
 
-TSharedRef<ITableRow> FSkeletonTreeVirtualBoneItem::MakeTreeRowWidget(const TSharedRef<STableViewBase>& InOwnerTable, const TAttribute<FText>& InFilterText)
+FSkeletonTreeVirtualBoneItem::FSkeletonTreeVirtualBoneItem(const FName& InBoneName, const TSharedRef<class ISkeletonTree>& InSkeletonTree)
+	: FSkeletonTreeItem(InSkeletonTree)
+	, BoneName(InBoneName)
 {
-	return
-		SNew(SSkeletonTreeRow, InOwnerTable)
-		.Item(SharedThis(this))
-		.FilterText(InFilterText);
+	static const FString BoneProxyPrefix(TEXT("VIRTUALBONEPROXY_"));
+
+	BoneProxy = NewObject<UBoneProxy>(GetTransientPackage(), *(BoneProxyPrefix + FString::Printf(TEXT("%p"), &InSkeletonTree.Get()) + InBoneName.ToString()));
+	BoneProxy->SetFlags(RF_Transactional);
+	BoneProxy->AddToRoot();
+	BoneProxy->BoneName = InBoneName;
+	TSharedPtr<IPersonaPreviewScene> PreviewScene = InSkeletonTree->GetPreviewScene();
+	if (PreviewScene.IsValid())
+	{
+		BoneProxy->SkelMeshComponent = PreviewScene->GetPreviewMeshComponent();
+	}
+}
+
+FSkeletonTreeVirtualBoneItem::~FSkeletonTreeVirtualBoneItem()
+{
+	if (BoneProxy)
+	{
+		BoneProxy->RemoveFromRoot();
+		BoneProxy = nullptr;
+	}
 }
 
 EVisibility FSkeletonTreeVirtualBoneItem::GetLODIconVisibility() const
@@ -171,4 +191,10 @@ EVisibility FSkeletonTreeVirtualBoneItem::GetVirtualBonePrefixVisibility() const
 {
 	return InlineWidget->IsInEditMode() ? EVisibility::Visible : EVisibility::Collapsed;
 }
+
+void FSkeletonTreeVirtualBoneItem::EnableBoneProxyTick(bool bEnable)
+{
+	BoneProxy->bIsTickable = bEnable;
+}
+
 #undef LOCTEXT_NAMESPACE
