Fix for crash in physics asset editor after garbage collection

Move bone proxies from rooting to FGCObject

diff --git a/Engine/Source/Editor/SkeletonEditor/Private/SSkeletonTree.cpp b/Engine/Source/Editor/SkeletonEditor/Private/SSkeletonTree.cpp
--- a/Engine/Source/Editor/SkeletonEditor/Private/SSkeletonTree.cpp
+++ b/Engine/Source/Editor/SkeletonEditor/Private/SSkeletonTree.cpp
@@ -170,7 +170,6 @@ void SSkeletonTree::Construct(const FArguments& InArgs, const TSharedRef<FEditab
 	}
 
 	BoneProxy = NewObject<UBoneProxy>(GetTransientPackage());
-	BoneProxy->AddToRoot();
 	BoneProxy->SkelMeshComponent = PreviewScene.IsValid() ? PreviewScene.Pin()->GetPreviewMeshComponent() : nullptr;
 	BoneProxy->bIsTickable = true;
 
@@ -281,12 +280,6 @@ SSkeletonTree::~SSkeletonTree()
 	{
 		EditableSkeleton.Pin()->UnregisterOnSkeletonHierarchyChanged(this);
 	}
-
-	if (BoneProxy)
-	{
-		BoneProxy->RemoveFromRoot();
-		BoneProxy = nullptr;
-	}
 }
 PRAGMA_ENABLE_DEPRECATION_WARNINGS
 
@@ -2105,4 +2098,9 @@ ESkeletonTreeFilterResult SSkeletonTree::HandleFilterSkeletonTreeItem(const FSke
 	return Result;
 }
 
+void SSkeletonTree::AddReferencedObjects( FReferenceCollector& Collector )
+{
+	Collector.AddReferencedObject(BoneProxy);
+}
+
 #undef LOCTEXT_NAMESPACE
diff --git a/Engine/Source/Editor/SkeletonEditor/Private/SSkeletonTree.h b/Engine/Source/Editor/SkeletonEditor/Private/SSkeletonTree.h
--- a/Engine/Source/Editor/SkeletonEditor/Private/SSkeletonTree.h
+++ b/Engine/Source/Editor/SkeletonEditor/Private/SSkeletonTree.h
@@ -24,6 +24,7 @@
 #include "SkeletonTreeBuilder.h"
 #include "SSearchBox.h"
 #include "EditorUndoClient.h"
+#include "GCObject.h"
 
 class FMenuBuilder;
 class FSkeletonTreeAttachedAssetItem;
@@ -42,7 +43,7 @@ struct FNotificationInfo;
 // SSkeletonTree
 
 PRAGMA_DISABLE_DEPRECATION_WARNINGS
-class SSkeletonTree : public ISkeletonTree, public FEditorUndoClient
+class SSkeletonTree : public ISkeletonTree, public FEditorUndoClient, public FGCObject
 {
 public:
 	SLATE_BEGIN_ARGS( SSkeletonTree )
@@ -106,6 +107,9 @@ public:
 	virtual void PostUndo(bool bSuccess) override;
 	virtual void PostRedo(bool bSuccess) override;
 
+	/** FGCObject interface */
+	virtual void AddReferencedObjects( FReferenceCollector& Collector ) override;
+
 	/** Creates the tree control and then populates */
 	void CreateTreeColumns();
 
diff --git a/Engine/Source/Editor/SkeletonEditor/Private/SkeletonTreeBoneItem.cpp b/Engine/Source/Editor/SkeletonEditor/Private/SkeletonTreeBoneItem.cpp
--- a/Engine/Source/Editor/SkeletonEditor/Private/SkeletonTreeBoneItem.cpp
+++ b/Engine/Source/Editor/SkeletonEditor/Private/SkeletonTreeBoneItem.cpp
@@ -32,7 +32,6 @@ FSkeletonTreeBoneItem::FSkeletonTreeBoneItem(const FName& InBoneName, const TSha
 
 	BoneProxy = NewObject<UBoneProxy>(GetTransientPackage(), *(BoneProxyPrefix + FString::Printf(TEXT("%p"), &InSkeletonTree.Get()) + InBoneName.ToString()));
 	BoneProxy->SetFlags(RF_Transactional);
-	BoneProxy->AddToRoot();
 	BoneProxy->BoneName = InBoneName;
 	TSharedPtr<IPersonaPreviewScene> PreviewScene = InSkeletonTree->GetPreviewScene();
 	if (PreviewScene.IsValid())
@@ -41,15 +40,6 @@ FSkeletonTreeBoneItem::FSkeletonTreeBoneItem(const FName& InBoneName, const TSha
 	}
 }
 
-FSkeletonTreeBoneItem::~FSkeletonTreeBoneItem()
-{
-	if (BoneProxy)
-	{
-		BoneProxy->RemoveFromRoot();
-		BoneProxy = nullptr;
-	}
-}
-
 EVisibility FSkeletonTreeBoneItem::GetLODIconVisibility() const
 {
 	if (bRequiredBone)
@@ -457,4 +447,9 @@ bool FSkeletonTreeBoneItem::IsBoneRequired(int32 MeshBoneIndex, UDebugSkelMeshCo
 	return Index != INDEX_NONE;
 }
 
+void FSkeletonTreeBoneItem::AddReferencedObjects( FReferenceCollector& Collector )
+{
+	Collector.AddReferencedObject(BoneProxy);
+}
+
 #undef LOCTEXT_NAMESPACE
\ No newline at end of file
diff --git a/Engine/Source/Editor/SkeletonEditor/Private/SkeletonTreeBoneItem.h b/Engine/Source/Editor/SkeletonEditor/Private/SkeletonTreeBoneItem.h
--- a/Engine/Source/Editor/SkeletonEditor/Private/SkeletonTreeBoneItem.h
+++ b/Engine/Source/Editor/SkeletonEditor/Private/SkeletonTreeBoneItem.h
@@ -15,16 +15,16 @@
 #include "Animation/Skeleton.h"
 #include "Widgets/Input/SComboButton.h"
 #include "BoneProxy.h"
+#include "GCObject.h"
 
 class UDebugSkelMeshComponent;
 
-class FSkeletonTreeBoneItem : public FSkeletonTreeItem
+class FSkeletonTreeBoneItem : public FSkeletonTreeItem, public FGCObject
 {
 public:
 	SKELETON_TREE_ITEM_TYPE(FSkeletonTreeBoneItem, FSkeletonTreeItem)
 
 	FSkeletonTreeBoneItem(const FName& InBoneName, const TSharedRef<class ISkeletonTree>& InSkeletonTree);
-	virtual ~FSkeletonTreeBoneItem();
 
 	/** ISkeletonTreeItem interface */
 	virtual void GenerateWidgetForNameColumn(TSharedPtr< SHorizontalBox > Box, const TAttribute<FText>& FilterText, FIsSelected InIsSelected) override;
@@ -36,6 +36,9 @@ public:
 	virtual FReply OnDragDetected(const FGeometry& MyGeometry, const FPointerEvent& MouseEvent) override;
 	virtual UObject* GetObject() const override { return BoneProxy; }
 
+	/** FGCObject interface */
+	virtual void AddReferencedObjects( FReferenceCollector& Collector ) override;
+
 	/** Check to see if the specified bone is weighted in the specified component */
 	bool IsBoneWeighted(int32 MeshBoneIndex, UDebugSkelMeshComponent* PreviewComponent);
 
diff --git a/Engine/Source/Editor/SkeletonEditor/Private/SkeletonTreeVirtualBoneItem.cpp b/Engine/Source/Editor/SkeletonEditor/Private/SkeletonTreeVirtualBoneItem.cpp
--- a/Engine/Source/Editor/SkeletonEditor/Private/SkeletonTreeVirtualBoneItem.cpp
+++ b/Engine/Source/Editor/SkeletonEditor/Private/SkeletonTreeVirtualBoneItem.cpp
@@ -21,7 +21,6 @@ FSkeletonTreeVirtualBoneItem::FSkeletonTreeVirtualBoneItem(const FName& InBoneNa
 
 	BoneProxy = NewObject<UBoneProxy>(GetTransientPackage(), *(BoneProxyPrefix + FString::Printf(TEXT("%p"), &InSkeletonTree.Get()) + InBoneName.ToString()));
 	BoneProxy->SetFlags(RF_Transactional);
-	BoneProxy->AddToRoot();
 	BoneProxy->BoneName = InBoneName;
 	TSharedPtr<IPersonaPreviewScene> PreviewScene = InSkeletonTree->GetPreviewScene();
 	if (PreviewScene.IsValid())
@@ -30,15 +29,6 @@ FSkeletonTreeVirtualBoneItem::FSkeletonTreeVirtualBoneItem(const FName& InBoneNa
 	}
 }
 
-FSkeletonTreeVirtualBoneItem::~FSkeletonTreeVirtualBoneItem()
-{
-	if (BoneProxy)
-	{
-		BoneProxy->RemoveFromRoot();
-		BoneProxy = nullptr;
-	}
-}
-
 EVisibility FSkeletonTreeVirtualBoneItem::GetLODIconVisibility() const
 {
 	return EVisibility::Visible;
@@ -197,4 +187,9 @@ void FSkeletonTreeVirtualBoneItem::EnableBoneProxyTick(bool bEnable)
 	BoneProxy->bIsTickable = bEnable;
 }
 
+void FSkeletonTreeVirtualBoneItem::AddReferencedObjects( FReferenceCollector& Collector )
+{
+	Collector.AddReferencedObject(BoneProxy);
+}
+
 #undef LOCTEXT_NAMESPACE
diff --git a/Engine/Source/Editor/SkeletonEditor/Private/SkeletonTreeVirtualBoneItem.h b/Engine/Source/Editor/SkeletonEditor/Private/SkeletonTreeVirtualBoneItem.h
--- a/Engine/Source/Editor/SkeletonEditor/Private/SkeletonTreeVirtualBoneItem.h
+++ b/Engine/Source/Editor/SkeletonEditor/Private/SkeletonTreeVirtualBoneItem.h
@@ -13,16 +13,16 @@
 #include "SkeletonTreeItem.h"
 #include "IEditableSkeleton.h"
 #include "BoneProxy.h"
+#include "GCObject.h"
 
 class SInlineEditableTextBlock;
 
-class FSkeletonTreeVirtualBoneItem : public FSkeletonTreeItem
+class FSkeletonTreeVirtualBoneItem : public FSkeletonTreeItem, public FGCObject
 {
 public:
 	SKELETON_TREE_ITEM_TYPE(FSkeletonTreeVirtualBoneItem, FSkeletonTreeItem)
 
 	FSkeletonTreeVirtualBoneItem(const FName& InBoneName, const TSharedRef<class ISkeletonTree>& InSkeletonTree);
-	~FSkeletonTreeVirtualBoneItem();
 
 	/** Builds the table row widget to display this info */
 	virtual void GenerateWidgetForNameColumn(TSharedPtr< SHorizontalBox > Box, const TAttribute<FText>& FilterText, FIsSelected InIsSelected) override;
@@ -33,6 +33,9 @@ public:
 	virtual void OnItemDoubleClicked() override;
 	virtual UObject* GetObject() const override { return BoneProxy; }
 
+	/** FGCObject interface */
+	virtual void AddReferencedObjects( FReferenceCollector& Collector ) override;
+
 	/** Return socket name as FText for display in skeleton tree */
 	FText GetVirtualBoneNameAsText() const { return FText::FromName(BoneName); }
 
