Fixed crash undoing the renaming of a deleted animation curve

diff --git a/Engine/Source/Editor/SkeletonEditor/Private/EditableSkeleton.cpp b/Engine/Source/Editor/SkeletonEditor/Private/EditableSkeleton.cpp
--- a/Engine/Source/Editor/SkeletonEditor/Private/EditableSkeleton.cpp
+++ b/Engine/Source/Editor/SkeletonEditor/Private/EditableSkeleton.cpp
@@ -416,12 +416,9 @@ void FEditableSkeleton::RemoveSmartnamesAndFixupAnimations(const FName& InContai
 	TArray<FAssetData> AnimationAssets;
 	GetAssetsContainingCurves(InContainerName, InNames, AnimationAssets);
 
-	bool bRemoved = true;
-
 	// AnimationAssets now only contains assets that are using the selected curve(s)
 	if (AnimationAssets.Num() > 0)
 	{
-		bRemoved = false; //need to warn user now
 		FString AssetMessage = LOCTEXT("DeleteCurveMessage", "Deleting curves will:\n\nRemove the curves from Animations and PoseAssets\nRemove poses using that curve name from PoseAssets.\n\nThe following assets will be modified. Continue?\n\n").ToString();
 
 		AnimationAssets.Sort([&](const FAssetData& A, const FAssetData& B)
@@ -443,7 +440,6 @@ void FEditableSkeleton::RemoveSmartnamesAndFixupAnimations(const FName& InContai
 
 		if (FMessageDialog::Open(EAppMsgType::YesNo, AssetMessageText, &AssetTitleText) == EAppReturnType::Yes)
 		{
-			bRemoved = true;
 			// Proceed to delete the curves
 			GWarn->BeginSlowTask(FText::Format(LOCTEXT("DeleteCurvesTaskDesc", "Deleting curve from skeleton {0}"), FText::FromString(Skeleton->GetName())), true);
 			FScopedTransaction Transaction(LOCTEXT("DeleteCurvesTransactionName", "Delete skeleton curve"));
@@ -486,13 +482,17 @@ void FEditableSkeleton::RemoveSmartnamesAndFixupAnimations(const FName& InContai
 				Seq->RequestSyncAnimRecompression();
 			}
 			GWarn->EndSlowTask();
+
+			// Remove names from skeleton
+			Skeleton->RemoveSmartnamesAndModify(InContainerName, InNames);
 		}
 	}
-
-	if (bRemoved && InNames.Num() > 0)
+	else if(InNames.Num() > 0)
 	{
+		FScopedTransaction Transaction(LOCTEXT("DeleteCurvesTransactionName", "Delete skeleton curve"));
+
 		// Remove names from skeleton
-		Skeleton->RemoveSmartnamesAndModify(InContainerName, InNames);
+		Skeleton->RemoveSmartnamesAndModify(InContainerName, InNames);	
 	}
 
 	OnSmartNameChanged.Broadcast(InContainerName);
diff --git a/Engine/Source/Runtime/Engine/Public/Animation/AnimCurveTypes.h b/Engine/Source/Runtime/Engine/Public/Animation/AnimCurveTypes.h
--- a/Engine/Source/Runtime/Engine/Public/Animation/AnimCurveTypes.h
+++ b/Engine/Source/Runtime/Engine/Public/Animation/AnimCurveTypes.h
@@ -360,7 +360,7 @@ struct FBaseBlendedCurve
 	/** Get Array Index by UID */
 	int32 GetArrayIndexByUID(USkeleton::AnimCurveUID InUid) const
 	{
-		int32 ArrayIndex = (*UIDToArrayIndexLUT)[InUid];
+		int32 ArrayIndex = (*UIDToArrayIndexLUT).IsValidIndex(InUid) ? (*UIDToArrayIndexLUT)[InUid] : MAX_uint16;
 		if (ArrayIndex != MAX_uint16)
 		{
 			return ArrayIndex;
@@ -373,7 +373,7 @@ struct FBaseBlendedCurve
 	{
 		check(bInitialized);
 
-		return ((*UIDToArrayIndexLUT)[InUid] != MAX_uint16);
+		return (GetArrayIndexByUID(InUid) != INDEX_NONE);
 	}
 	
 	/** Get Valid Element Count from given UIDToArrayIndexLUT */
