Skeletal Mesh Editor: Fixed Preview Asset becoming detached when socket is renamed, due to passing names by reference rather than value.

diff --git a/Engine/Source/Editor/SkeletonEditor/Private/EditableSkeleton.cpp b/Engine/Source/Editor/SkeletonEditor/Private/EditableSkeleton.cpp
--- a/Engine/Source/Editor/SkeletonEditor/Private/EditableSkeleton.cpp
+++ b/Engine/Source/Editor/SkeletonEditor/Private/EditableSkeleton.cpp
@@ -247,7 +247,7 @@ bool FEditableSkeleton::AddSmartname(const FName& InContainerName, const FName&
 	return false;
 }
 
-void FEditableSkeleton::RenameSmartname(const FName& InContainerName, SmartName::UID_Type InNameUid, const FName& InNewName)
+void FEditableSkeleton::RenameSmartname(const FName InContainerName, SmartName::UID_Type InNameUid, const FName InNewName)
 {
 	FSmartName CurveToRename;
 	if (!Skeleton->GetSmartNameByUID(USkeleton::AnimCurveMappingName, InNameUid, CurveToRename))
@@ -560,7 +560,7 @@ static USkeletalMeshSocket* FindSocket(const FName& InSocketName, USkeletalMesh*
 	return Socket;
 }
 
-void FEditableSkeleton::RenameSocket(const FName& OldSocketName, const FName& NewSocketName, USkeletalMesh* InSkeletalMesh)
+void FEditableSkeleton::RenameSocket(const FName OldSocketName, const FName NewSocketName, USkeletalMesh* InSkeletalMesh)
 {
 	USkeletalMeshSocket* SocketData = FindSocket(OldSocketName, InSkeletalMesh, Skeleton);
 
@@ -1007,7 +1007,7 @@ bool FEditableSkeleton::DoesVirtualBoneAlreadyExist(const FString& InVBName) con
 	return false;
 }
 
-void FEditableSkeleton::RenameVirtualBone(const FName& OriginalName, const FName& InVBName)
+void FEditableSkeleton::RenameVirtualBone(const FName OriginalName, const FName InVBName)
 {
 	FScopedTransaction Transaction(LOCTEXT("RenameVirtualBone", "Rename Virtual Bone in Skeleton"));
 	Skeleton->RenameVirtualBone(OriginalName, InVBName);
@@ -1122,7 +1122,7 @@ void FEditableSkeleton::AddSyncMarker(FName NewName)
 	Skeleton->RegisterMarkerName(NewName);
 }
 
-int32 FEditableSkeleton::RenameNotify(const FName& NewName, const FName& OldName)
+int32 FEditableSkeleton::RenameNotify(const FName NewName, const FName OldName)
 {
 	const FScopedTransaction Transaction(LOCTEXT("RenameAnimNotify", "Rename Anim Notify"));
 	Skeleton->Modify();
@@ -1244,7 +1244,7 @@ void FEditableSkeleton::SetAdditionalPreviewSkeletalMeshes(class UDataAsset* InP
 	}
 }
 
-void FEditableSkeleton::RenameRetargetSource(const FName& InOldName, const FName& InNewName)
+void FEditableSkeleton::RenameRetargetSource(const FName InOldName, const FName InNewName)
 {
 	const FScopedTransaction Transaction(LOCTEXT("RetargetSourceWindow_Rename", "Rename Retarget Source"));
 
@@ -1611,7 +1611,7 @@ void FEditableSkeleton::DeleteSlotGroup(const FName& InGroupName)
 	Skeleton->RemoveSlotGroup(InGroupName);
 }
 
-void FEditableSkeleton::RenameSlotName(const FName& InOldSlotName, const FName& InNewSlotName)
+void FEditableSkeleton::RenameSlotName(const FName InOldSlotName, const FName InNewSlotName)
 {
 	const FScopedTransaction Transaction(LOCTEXT("RenameSlotName", "Rename Slot Name"));
 	Skeleton->Modify();
diff --git a/Engine/Source/Editor/SkeletonEditor/Private/EditableSkeleton.h b/Engine/Source/Editor/SkeletonEditor/Private/EditableSkeleton.h
--- a/Engine/Source/Editor/SkeletonEditor/Private/EditableSkeleton.h
+++ b/Engine/Source/Editor/SkeletonEditor/Private/EditableSkeleton.h
@@ -40,23 +40,23 @@ public:
 	virtual void DeleteSyncMarkers(const TArray<FName>& ISyncMarkerNames) override;
 	virtual void AddNotify(FName NewName) override;
 	virtual void AddSyncMarker(FName NewName) override;
-	virtual int32 RenameNotify(const FName& NewName, const FName& OldName) override;
+	virtual int32 RenameNotify(const FName NewName, const FName OldName) override;
 	virtual void BroadcastNotifyChanged() override;
 	virtual void GetCompatibleAnimSequences(TArray<struct FAssetData>& OutAssets) override;
-	virtual void RenameSocket(const FName& OldSocketName, const FName& NewSocketName, USkeletalMesh* InSkeletalMesh) override;
+	virtual void RenameSocket(const FName OldSocketName, const FName NewSocketName, USkeletalMesh* InSkeletalMesh) override;
 	virtual void SetSocketParent(const FName& SocketName, const FName& NewParentName, USkeletalMesh* InSkeletalMesh) override;
 	virtual bool DoesSocketAlreadyExist(const class USkeletalMeshSocket* InSocket, const FText& InSocketName, ESocketParentType SocketParentType, USkeletalMesh* InSkeletalMesh) const override;
 	virtual bool DoesVirtualBoneAlreadyExist(const FString& InVBName) const override;
-	virtual void RenameVirtualBone(const FName& OriginalName, const FName& InVBName) override;
+	virtual void RenameVirtualBone(const FName OriginalName, const FName InVBName) override;
 	virtual bool AddSmartname(const FName& InContainerName, const FName& InNewName, FSmartName& OutSmartName) override;
-	virtual void RenameSmartname(const FName& InContainerName, SmartName::UID_Type InNameUid, const FName& InNewName) override;
+	virtual void RenameSmartname(const FName InContainerName, SmartName::UID_Type InNameUid, const FName InNewName) override;
 	virtual void RemoveSmartnamesAndFixupAnimations(const FName& InContainerName, const TArray<FName>& InNames) override;
 	virtual void SetCurveMetaDataMaterial(const FSmartName& CurveName, bool bOverrideMaterial) override;
 	virtual void SetCurveMetaBoneLinks(const FSmartName& CurveName, TArray<FBoneReference>& BoneLinks, uint8 InMaxLOD) override;
 	virtual void SetPreviewMesh(class USkeletalMesh* InSkeletalMesh) override;
 	virtual void LoadAdditionalPreviewSkeletalMeshes() override;
 	virtual void SetAdditionalPreviewSkeletalMeshes(class UDataAsset* InPreviewCollectionAsset) override;
-	virtual void RenameRetargetSource(const FName& InOldName, const FName& InNewName) override;
+	virtual void RenameRetargetSource(const FName InOldName, const FName InNewName) override;
 	virtual void AddRetargetSource(const FName& InName, USkeletalMesh* InReferenceMesh) override;
 	virtual void DeleteRetargetSources(const TArray<FName>& InRetargetSourceNames) override;
 	virtual void RefreshRetargetSources(const TArray<FName>& InRetargetSourceNames) override;
@@ -70,7 +70,7 @@ public:
 	virtual void SetSlotGroupName(const FName& InSlotName, const FName& InGroupName) override;
 	virtual void DeleteSlotName(const FName& InSlotName) override;
 	virtual void DeleteSlotGroup(const FName& InGroupName) override;
-	virtual void RenameSlotName(const FName& InOldSlotName, const FName& InNewSlotName) override;
+	virtual void RenameSlotName(const FName InOldSlotName, const FName InNewSlotName) override;
 	virtual FDelegateHandle RegisterOnSmartNameChanged(const FOnSmartNameChanged::FDelegate& InOnSmartNameChanged) override;
 	virtual void UnregisterOnSmartNameChanged(FDelegateHandle InHandle) override;
 	virtual void RegisterOnNotifiesChanged(const FSimpleMulticastDelegate::FDelegate& InDelegate);
diff --git a/Engine/Source/Editor/SkeletonEditor/Public/IEditableSkeleton.h b/Engine/Source/Editor/SkeletonEditor/Public/IEditableSkeleton.h
--- a/Engine/Source/Editor/SkeletonEditor/Public/IEditableSkeleton.h
+++ b/Engine/Source/Editor/SkeletonEditor/Public/IEditableSkeleton.h
@@ -45,7 +45,7 @@ public:
 	virtual class USkeletalMeshSocket* DuplicateSocket(const struct FSelectedSocketInfo& SocketInfoToDuplicate, const FName& NewParentBoneName, USkeletalMesh* InSkeletalMesh) = 0;
 
 	/** Rename a socket on the mesh or the skeleton */
-	virtual void RenameSocket(const FName& OldSocketName, const FName& NewSocketName, USkeletalMesh* InSkeletalMesh) = 0;
+	virtual void RenameSocket(const FName OldSocketName, const FName NewSocketName, USkeletalMesh* InSkeletalMesh) = 0;
 
 	/** Set the parent of a socket */
 	virtual void SetSocketParent(const FName& SocketName, const FName& NewParentName, USkeletalMesh* InSkeletalMesh) = 0;
@@ -57,7 +57,7 @@ public:
 	virtual bool AddSmartname(const FName& InContainerName, const FName& InNewName, FSmartName& OutSmartName) = 0;
 
 	/** Rename the specified smart name */
-	virtual void RenameSmartname(const FName& InContainerName, SmartName::UID_Type InNameUid, const FName& InNewName) = 0;
+	virtual void RenameSmartname(const FName InContainerName, SmartName::UID_Type InNameUid, const FName InNewName) = 0;
 
 	/** Remove all the specified smart names and fixup animations that use them */
 	virtual void RemoveSmartnamesAndFixupAnimations(const FName& InContainerName, const TArray<FName>& InNames) = 0;
@@ -106,7 +106,7 @@ public:
 	 * Rename a notify
 	 * @return the number of animations modified
 	 */	
-	virtual int32 RenameNotify(const FName& NewName, const FName& OldName) = 0;
+	virtual int32 RenameNotify(const FName NewName, const FName OldName) = 0;
 
 	/**
 	* Add a sync marker
@@ -129,7 +129,7 @@ public:
 	virtual void SetAdditionalPreviewSkeletalMeshes(class UDataAsset* InPreviewCollectionAsset) = 0;
 
 	/** Rename the specified retarget source */
-	virtual void RenameRetargetSource(const FName& InOldName, const FName& InNewName) = 0;
+	virtual void RenameRetargetSource(const FName InOldName, const FName InNewName) = 0;
 
 	/** Add a retarget source */
 	virtual void AddRetargetSource(const FName& InName, USkeletalMesh* InReferenceMesh) = 0;
@@ -171,7 +171,7 @@ public:
 	virtual void DeleteSlotGroup(const FName& InGroupName) = 0;
 
 	/** Rename a slot name */
-	virtual void RenameSlotName(const FName& InOldSlotName, const FName& InNewSlotName) = 0;
+	virtual void RenameSlotName(const FName InOldSlotName, const FName InNewSlotName) = 0;
 
 	/** Register a delegate to be called when a set of smart names are removed */
 	virtual FDelegateHandle RegisterOnSmartNameChanged(const FOnSmartNameChanged::FDelegate& InOnSmartNameChanged) = 0;
@@ -195,7 +195,7 @@ public:
 	virtual bool DoesVirtualBoneAlreadyExist(const FString& InVBName) const = 0;
 
 	/** Rename an existing virtual bone */
-	virtual void RenameVirtualBone(const FName& OriginalName, const FName& InVBName) = 0;
+	virtual void RenameVirtualBone(const FName OriginalName, const FName InVBName) = 0;
 
 	/** Broadcasts tree refresh delegate */
 	virtual void RefreshBoneTree() = 0;
