Prevent TArrayView being instantiated with an rvalue initializer list.
Fixes to code which now violates that assumption.

diff --git a/Engine/Source/Editor/MovieSceneTools/Public/KeyframeTrackEditor.h b/Engine/Source/Editor/MovieSceneTools/Public/KeyframeTrackEditor.h
--- a/Engine/Source/Editor/MovieSceneTools/Public/KeyframeTrackEditor.h
+++ b/Engine/Source/Editor/MovieSceneTools/Public/KeyframeTrackEditor.h
@@ -377,6 +377,13 @@ protected:
 		return KeyPropertyResult;
 	}
 
+	FORCEINLINE FKeyPropertyResult AddKeysToObjects(
+		std::initializer_list<UObject*> ObjectsToKey, FFrameNumber KeyTime, FGeneratedTrackKeys& GeneratedKeys,
+		ESequencerKeyMode KeyMode, TSubclassOf<UMovieSceneTrack> TrackClass, FName PropertyName,
+		TFunction<void(TrackType*)> OnInitializeNewTrack)
+	{
+		return AddKeysToObjects(MakeArrayView(ObjectsToKey), KeyTime, GeneratedKeys, KeyMode, TrackClass, PropertyName, MoveTemp(OnInitializeNewTrack));
+	}
 
 private:
 
diff --git a/Engine/Source/Editor/SkeletonEditor/Public/ISkeletonTreeBuilder.h b/Engine/Source/Editor/SkeletonEditor/Public/ISkeletonTreeBuilder.h
--- a/Engine/Source/Editor/SkeletonEditor/Public/ISkeletonTreeBuilder.h
+++ b/Engine/Source/Editor/SkeletonEditor/Public/ISkeletonTreeBuilder.h
@@ -27,6 +27,18 @@ struct SKELETONEDITOR_API FSkeletonTreeBuilderOutput
 	 */
 	void Add(const TSharedPtr<class ISkeletonTreeItem>& InItem, const FName& InParentName, TArrayView<const FName> InParentTypes, bool bAddToHead = false);
 
+	/** 
+	 * Add an item to the output
+	 * @param	InItem			The item to add
+	 * @param	InParentName	The name of the item's parent
+	 * @param	InParentTypes	The types of items to search. If this is empty all items will be searched.
+	 * @param	bAddToHead		Whether to add the item to the start or end of the parent's children array
+	 */
+	FORCEINLINE void Add(const TSharedPtr<class ISkeletonTreeItem>& InItem, const FName& InParentName, std::initializer_list<FName> InParentTypes, bool bAddToHead = false)
+	{
+		Add(InItem, InParentName, MakeArrayView(InParentTypes), bAddToHead);
+	}
+
 	/** 
 	 * Add an item to the output
 	 * @param	InItem			The item to add
@@ -44,6 +56,17 @@ struct SKELETONEDITOR_API FSkeletonTreeBuilderOutput
 	 */
 	TSharedPtr<class ISkeletonTreeItem> Find(const FName& InName, TArrayView<const FName> InTypes) const;
 
+	/** 
+	 * Find the item with the specified name
+	 * @param	InName	The item's name
+	 * @param	InTypes	The types of items to search. If this is empty all items will be searched.
+	 * @return the item found, or an invalid ptr if it was not found.
+	 */
+	FORCEINLINE TSharedPtr<class ISkeletonTreeItem> Find(const FName& InName, std::initializer_list<FName> InTypes) const
+	{
+		return Find(InName, MakeArrayView(InTypes));
+	}
+
 	/** 
 	 * Find the item with the specified name
 	 * @param	InName	The item's name
diff --git a/Engine/Source/Runtime/Core/Public/Containers/ArrayView.h b/Engine/Source/Runtime/Core/Public/Containers/ArrayView.h
--- a/Engine/Source/Runtime/Core/Public/Containers/ArrayView.h
+++ b/Engine/Source/Runtime/Core/Public/Containers/ArrayView.h
@@ -3,6 +3,7 @@
 #pragma once
 
 #include "CoreTypes.h"
+#include "Templates/IsInitializerList.h"
 #include "Templates/PointerIsConvertibleFromTo.h"
 #include "Misc/AssertionMacros.h"
 #include "Templates/UnrealTypeTraits.h"
@@ -137,9 +138,9 @@ public:
 	 *
 	 * @param List The initializer list to view.
 	 */
-	template <typename OtherElementType,
-		typename = typename TEnableIf<TIsCompatibleElementType<const OtherElementType>::Value>::Type>
-	FORCEINLINE TArrayView(std::initializer_list<OtherElementType> List)
+	template <typename InitializerListType,
+		typename = typename TEnableIf<TAnd<TIsInitializerList<InitializerListType>, TIsCompatibleElementType<const typename InitializerListType::value_type>>::Value>::Type>
+	FORCEINLINE TArrayView(InitializerListType& List)
 		: DataPtr(&*List.begin())
 		, ArrayNum(List.size())
 	{
@@ -596,10 +597,13 @@ TArrayView<const ElementType> MakeArrayView(const TArray<ElementType, Allocator>
 	return TArrayView<const ElementType>(Other);
 }
 
-template<typename ElementType>
-TArrayView<ElementType> MakeArrayView(std::initializer_list<ElementType> List)
+template <typename InitializerListType>
+typename TEnableIf<
+	TIsInitializerList<InitializerListType>::Value,
+	TArrayView<const typename InitializerListType::value_type>
+>::Type MakeArrayView(InitializerListType& List)
 {
-	return TArrayView<ElementType>(List);
+	return TArrayView<const typename InitializerListType::value_type>(List);
 }
 
 template<typename ElementType>
diff --git a/Engine/Source/Runtime/Core/Public/ProfilingDebugging/Histogram.h b/Engine/Source/Runtime/Core/Public/ProfilingDebugging/Histogram.h
--- a/Engine/Source/Runtime/Core/Public/ProfilingDebugging/Histogram.h
+++ b/Engine/Source/Runtime/Core/Public/ProfilingDebugging/Histogram.h
@@ -35,6 +35,12 @@ struct CORE_API FHistogram
 	/** Inits histogram with the specified bin boundaries, with the final bucket extending to infinity (e.g., passing in 0,5 creates a [0..5) bucket and a [5..infinity) bucket) */
 	void InitFromArray(TArrayView<const double> Thresholds);
 
+	/** Inits histogram with the specified bin boundaries, with the final bucket extending to infinity (e.g., passing in 0,5 creates a [0..5) bucket and a [5..infinity) bucket) */
+	FORCEINLINE void InitFromArray(std::initializer_list<double> Thresholds)
+	{
+		InitFromArray(MakeArrayView(Thresholds));
+	}
+
 	/** Resets measurements, without resetting the configured bins. */
 	void Reset();
 
diff --git a/Engine/Source/Runtime/Core/Public/Templates/IsInitializerList.h b/Engine/Source/Runtime/Core/Public/Templates/IsInitializerList.h
new file mode 100644
--- /dev/null
+++ b/Engine/Source/Runtime/Core/Public/Templates/IsInitializerList.h
@@ -0,0 +1,24 @@
+// Copyright 1998-2019 Epic Games, Inc. All Rights Reserved.
+
+#pragma once
+
+#include <initializer_list>
+
+/**
+ * Traits class which tests if a type is an initializer list.
+ */
+template <typename T>
+struct TIsInitializerList
+{
+	static constexpr bool Value = false;
+};
+
+template <typename T>
+struct TIsInitializerList<std::initializer_list<T>>
+{
+	static constexpr bool Value = true;
+};
+
+template <typename T> struct TIsInitializerList<const          T> { enum { Value = TIsInitializerList<T>::Value }; };
+template <typename T> struct TIsInitializerList<      volatile T> { enum { Value = TIsInitializerList<T>::Value }; };
+template <typename T> struct TIsInitializerList<const volatile T> { enum { Value = TIsInitializerList<T>::Value }; };
diff --git a/Engine/Source/Runtime/D3D12RHI/Private/D3D12RayTracing.cpp b/Engine/Source/Runtime/D3D12RHI/Private/D3D12RayTracing.cpp
--- a/Engine/Source/Runtime/D3D12RHI/Private/D3D12RayTracing.cpp
+++ b/Engine/Source/Runtime/D3D12RHI/Private/D3D12RayTracing.cpp
@@ -852,7 +852,7 @@ public:
 		// Use hit shaders from initializer or fall back to default if none were provided
 
 		FRayTracingShaderRHIParamRef DefaultHitShader = GetBuildInRayTracingShader<FDefaultMainCHS>();
-		TArrayView<const FRayTracingShaderRHIParamRef> DefaultHitGroupTable = { DefaultHitShader };
+		TArrayView<const FRayTracingShaderRHIParamRef> DefaultHitGroupTable = MakeArrayView(&DefaultHitShader, 1);
 
 		const TArrayView<const FRayTracingShaderRHIParamRef>& InitializerHitGroups = Initializer.GetHitGroupTable().Num()
 			? Initializer.GetHitGroupTable()
diff --git a/Engine/Source/Runtime/MovieScene/Private/Tests/MovieSceneSegmentCompilerTests.cpp b/Engine/Source/Runtime/MovieScene/Private/Tests/MovieSceneSegmentCompilerTests.cpp
--- a/Engine/Source/Runtime/MovieScene/Private/Tests/MovieSceneSegmentCompilerTests.cpp
+++ b/Engine/Source/Runtime/MovieScene/Private/Tests/MovieSceneSegmentCompilerTests.cpp
@@ -122,6 +122,11 @@ namespace Impl
 		}
 	}
 
+	FORCEINLINE void AssertSegmentAtTime(FAutomationTestBase* Test, FMovieSceneEvaluationTrack& InTrack, FFrameNumber InTime, std::initializer_list<FSectionEvaluationData> ExpectedImpls)
+	{
+		AssertSegmentAtTime(Test, InTrack, InTime, MakeArrayView(ExpectedImpls));
+	}
+
 	FORCEINLINE_DEBUGGABLE void AssertEvaluationTree(FAutomationTestBase* Test, const TMovieSceneEvaluationTree<int32>& Tree, TArrayView<const FEvaluationTreeIteratorResult> Expected)
 	{
 		int32 Index = 0;
diff --git a/Engine/Source/Runtime/MovieScene/Public/Evaluation/MovieSceneSegment.h b/Engine/Source/Runtime/MovieScene/Public/Evaluation/MovieSceneSegment.h
--- a/Engine/Source/Runtime/MovieScene/Public/Evaluation/MovieSceneSegment.h
+++ b/Engine/Source/Runtime/MovieScene/Public/Evaluation/MovieSceneSegment.h
@@ -163,6 +163,11 @@ struct FMovieSceneSegment
 		}
 	}
 
+	FMovieSceneSegment(const TRange<FFrameNumber>& InRange, std::initializer_list<FSectionEvaluationData> InApplicationImpls)
+		: FMovieSceneSegment(InRange, MakeArrayView(InApplicationImpls))
+	{
+	}
+
 	friend bool operator==(const FMovieSceneSegment& A, const FMovieSceneSegment& B)
 	{
 		return A.Range == B.Range && A.ID == B.ID && A.bAllowEmpty == B.bAllowEmpty && A.Impls == B.Impls;
diff --git a/Engine/Source/Runtime/Renderer/Private/SceneCaptureRendering.cpp b/Engine/Source/Runtime/Renderer/Private/SceneCaptureRendering.cpp
--- a/Engine/Source/Runtime/Renderer/Private/SceneCaptureRendering.cpp
+++ b/Engine/Source/Runtime/Renderer/Private/SceneCaptureRendering.cpp
@@ -660,7 +660,7 @@ static FSceneRenderer* CreateSceneRendererForSceneCapture(
 	SetupViewVamilyForSceneCapture(
 		ViewFamily,
 		SceneCaptureComponent,
-		{ SceneCaptureViewInfo },
+		MakeArrayView(&SceneCaptureViewInfo, 1),
 		MaxViewDistance, 
 		bCaptureSceneColor,
 		/* bIsPlanarReflection = */ false,
diff --git a/Engine/Source/Runtime/VulkanRHI/Private/VulkanDevice.h b/Engine/Source/Runtime/VulkanRHI/Private/VulkanDevice.h
--- a/Engine/Source/Runtime/VulkanRHI/Private/VulkanDevice.h
+++ b/Engine/Source/Runtime/VulkanRHI/Private/VulkanDevice.h
@@ -305,6 +305,11 @@ private:
 	void MapFormatSupport(EPixelFormat UEFormat, VkFormat VulkanFormat, int32 BlockBytes);
 	void SetComponentMapping(EPixelFormat UEFormat, VkComponentSwizzle r, VkComponentSwizzle g, VkComponentSwizzle b, VkComponentSwizzle a);
 
+	FORCEINLINE void MapFormatSupportWithFallback(EPixelFormat UEFormat, VkFormat VulkanFormat, std::initializer_list<VkFormat> FallbackTextureFormats)
+	{
+		MapFormatSupportWithFallback(UEFormat, VulkanFormat, MakeArrayView(FallbackTextureFormats));
+	}
+
 	void SubmitCommands(FVulkanCommandListContext* Context);
 
 
