Avoid building more then once a skeletalmesh when we remove bones to all LODs
#rnx

diff --git a/Engine/Source/Editor/SkeletonEditor/Private/SSkeletonTree.cpp b/Engine/Source/Editor/SkeletonEditor/Private/SSkeletonTree.cpp
--- a/Engine/Source/Editor/SkeletonEditor/Private/SSkeletonTree.cpp
+++ b/Engine/Source/Editor/SkeletonEditor/Private/SSkeletonTree.cpp
@@ -1039,46 +1039,52 @@ void SSkeletonTree::RemoveFromLOD(int32 LODIndex, bool bIncludeSelected, bool bI
 
 		TArray<FName> BonesToRemove;
 
-		for (const TSharedPtr<FSkeletonTreeBoneItem>& Item : TreeSelection.GetSelectedItems<FSkeletonTreeBoneItem>())
+		//Scoped post edit change
 		{
-			FName BoneName = Item->GetRowItemName();
-			int32 BoneIndex = RefSkeleton.FindBoneIndex(BoneName);
-			if (BoneIndex != INDEX_NONE)
+			FScopedSkeletalMeshPostEditChange ScopedPostEditChange(PreviewMeshComponent->SkeletalMesh);
+
+			for (const TSharedPtr<FSkeletonTreeBoneItem>& Item : TreeSelection.GetSelectedItems<FSkeletonTreeBoneItem>())
 			{
-				if (bIncludeSelected)
-				{
-					PreviewMeshComponent->SkeletalMesh->AddBoneToReductionSetting(LODIndex, BoneName);
-					BonesToRemove.AddUnique(BoneName);
-				}
-				else
+				FName BoneName = Item->GetRowItemName();
+				int32 BoneIndex = RefSkeleton.FindBoneIndex(BoneName);
+				if (BoneIndex != INDEX_NONE)
 				{
-					for (int32 ChildIndex = BoneIndex + 1; ChildIndex < RefSkeleton.GetRawBoneNum(); ++ChildIndex)
+					if (bIncludeSelected)
+					{
+						PreviewMeshComponent->SkeletalMesh->AddBoneToReductionSetting(LODIndex, BoneName);
+						BonesToRemove.AddUnique(BoneName);
+					}
+					else
 					{
-						if (RefSkeleton.GetParentIndex(ChildIndex) == BoneIndex)
+						for (int32 ChildIndex = BoneIndex + 1; ChildIndex < RefSkeleton.GetRawBoneNum(); ++ChildIndex)
 						{
-							FName ChildBoneName = RefSkeleton.GetBoneName(ChildIndex);
-							PreviewMeshComponent->SkeletalMesh->AddBoneToReductionSetting(LODIndex, ChildBoneName);
-							BonesToRemove.AddUnique(ChildBoneName);
+							if (RefSkeleton.GetParentIndex(ChildIndex) == BoneIndex)
+							{
+								FName ChildBoneName = RefSkeleton.GetBoneName(ChildIndex);
+								PreviewMeshComponent->SkeletalMesh->AddBoneToReductionSetting(LODIndex, ChildBoneName);
+								BonesToRemove.AddUnique(ChildBoneName);
+							}
 						}
 					}
 				}
 			}
-		}
 
-		int32 TotalLOD = PreviewMeshComponent->SkeletalMesh->GetLODNum();
-		IMeshUtilities& MeshUtilities = FModuleManager::Get().LoadModuleChecked<IMeshUtilities>("MeshUtilities");
+			int32 TotalLOD = PreviewMeshComponent->SkeletalMesh->GetLODNum();
+			IMeshUtilities& MeshUtilities = FModuleManager::Get().LoadModuleChecked<IMeshUtilities>("MeshUtilities");
 
-		if (bIncludeBelowLODs)
-		{
-			for (int32 Index = LODIndex + 1; Index < TotalLOD; ++Index)
+		
+			if (bIncludeBelowLODs)
 			{
-				MeshUtilities.RemoveBonesFromMesh(PreviewMeshComponent->SkeletalMesh, Index, &BonesToRemove);
-				PreviewMeshComponent->SkeletalMesh->AddBoneToReductionSetting(Index, BonesToRemove);
+				for (int32 Index = LODIndex + 1; Index < TotalLOD; ++Index)
+				{
+					MeshUtilities.RemoveBonesFromMesh(PreviewMeshComponent->SkeletalMesh, Index, &BonesToRemove);
+					PreviewMeshComponent->SkeletalMesh->AddBoneToReductionSetting(Index, BonesToRemove);
+				}
 			}
-		}
 
-		// remove from current LOD
-		MeshUtilities.RemoveBonesFromMesh(PreviewMeshComponent->SkeletalMesh, LODIndex, &BonesToRemove);
+			// remove from current LOD
+			MeshUtilities.RemoveBonesFromMesh(PreviewMeshComponent->SkeletalMesh, LODIndex, &BonesToRemove);
+		}
 		// update UI to reflect the change
 		OnLODSwitched();
 	}
