Fix Virtual Bone renaming

diff --git a/Engine/Source/Editor/SkeletonEditor/Private/SkeletonTreeVirtualBoneItem.cpp b/Engine/Source/Editor/SkeletonEditor/Private/SkeletonTreeVirtualBoneItem.cpp
--- a/Engine/Source/Editor/SkeletonEditor/Private/SkeletonTreeVirtualBoneItem.cpp
+++ b/Engine/Source/Editor/SkeletonEditor/Private/SkeletonTreeVirtualBoneItem.cpp
@@ -60,7 +60,7 @@ void FSkeletonTreeVirtualBoneItem::GenerateWidgetForNameColumn(TSharedPtr< SHori
 						.HighlightText(FilterText)
 						.Font(this, &FSkeletonTreeVirtualBoneItem::GetBoneTextFont)
 						.ToolTipText(ToolTip)
-						.OnBeginTextEdit(this, &FSkeletonTreeVirtualBoneItem::OnVirtualBoneNameEditing)
+						.OnEnterEditingMode(this, &FSkeletonTreeVirtualBoneItem::OnVirtualBoneNameEditing)
 						.OnVerifyTextChanged(this, &FSkeletonTreeVirtualBoneItem::OnVerifyBoneNameChanged)
 						.OnTextCommitted(this, &FSkeletonTreeVirtualBoneItem::OnCommitVirtualBoneName)
 						.IsSelected(InIsSelected);
@@ -128,10 +128,10 @@ void FSkeletonTreeVirtualBoneItem::RequestRename()
 	OnRenameRequested.ExecuteIfBound();
 }
 
-void FSkeletonTreeVirtualBoneItem::OnVirtualBoneNameEditing(const FText& OriginalText)
+void FSkeletonTreeVirtualBoneItem::OnVirtualBoneNameEditing()
 {
 	CachedBoneNameForRename = BoneName;
-	BoneName = VirtualBoneNameHelpers::RemoveVirtualBonePrefix(OriginalText.ToString());
+	BoneName = VirtualBoneNameHelpers::RemoveVirtualBonePrefix(BoneName.ToString());
 }
 
 bool FSkeletonTreeVirtualBoneItem::OnVerifyBoneNameChanged(const FText& InText, FText& OutErrorMessage)
diff --git a/Engine/Source/Editor/SkeletonEditor/Private/SkeletonTreeVirtualBoneItem.h b/Engine/Source/Editor/SkeletonEditor/Private/SkeletonTreeVirtualBoneItem.h
--- a/Engine/Source/Editor/SkeletonEditor/Private/SkeletonTreeVirtualBoneItem.h
+++ b/Engine/Source/Editor/SkeletonEditor/Private/SkeletonTreeVirtualBoneItem.h
@@ -44,7 +44,7 @@ public:
 
 private:
 	/** Called when we are about to rename a virtual bone */
-	void OnVirtualBoneNameEditing(const FText& OriginalText);
+	void OnVirtualBoneNameEditing();
 
 	/** Called when user is renaming this bone to verify the name **/
 	bool OnVerifyBoneNameChanged(const FText& InText, FText& OutErrorMessage);
