Fixed a bug that required the blend profile editor to require the user to press the Enter key in order for the per bone value change to be committed.
Pressing tab or clicking away from the edit field would not update the runtime.
This is fixed now.

diff --git a/Engine/Source/Editor/SkeletonEditor/Private/SkeletonTreeBoneItem.cpp b/Engine/Source/Editor/SkeletonEditor/Private/SkeletonTreeBoneItem.cpp
--- a/Engine/Source/Editor/SkeletonEditor/Private/SkeletonTreeBoneItem.cpp
+++ b/Engine/Source/Editor/SkeletonEditor/Private/SkeletonTreeBoneItem.cpp
@@ -335,10 +335,7 @@ FText FSkeletonTreeBoneItem::GetBoneToolTip()
 
 void FSkeletonTreeBoneItem::OnBlendSliderCommitted(float NewValue, ETextCommit::Type CommitType)
 {
-	if(CommitType == ETextCommit::OnEnter)
-	{
-		SetBoneBlendProfileScale(NewValue, false);
-	}
+	SetBoneBlendProfileScale(NewValue, false);
 }
 
 void FSkeletonTreeBoneItem::HandleDragEnter(const FDragDropEvent& DragDropEvent)
