Skeleton Tree Style Updates

diff --git a/Engine/Source/Editor/EditorStyle/Private/StarshipStyle.cpp b/Engine/Source/Editor/EditorStyle/Private/StarshipStyle.cpp
--- a/Engine/Source/Editor/EditorStyle/Private/StarshipStyle.cpp
+++ b/Engine/Source/Editor/EditorStyle/Private/StarshipStyle.cpp
@@ -4976,22 +4976,32 @@ void FStarshipEditorStyle::FStyle::SetupPersonaStyle()
 		Set("SkeletonTree.SkeletonSocket", new IMAGE_BRUSH("Persona/SkeletonTree/icon_SocketG_16px", Icon16x16));
 		Set("SkeletonTree.MeshSocket", new IMAGE_BRUSH("Persona/SkeletonTree/icon_SocketC_16px", Icon16x16));
 		Set("SkeletonTree.LODBone", new IMAGE_BRUSH(TEXT("Persona/SkeletonTree/icon_LODBone_16x"), Icon16x16));
+		Set("SkeletonTree.Bone", 	new IMAGE_BRUSH_SVG(TEXT("Starship/Animation/Bone"), Icon16x16) );
+		Set("SkeletonTree.BoneNonWeighted", new IMAGE_BRUSH_SVG(TEXT("Starship/Animation/BoneNonWeighted"), Icon16x16) );
 		Set("SkeletonTree.NonRequiredBone", new IMAGE_BRUSH(TEXT("Persona/SkeletonTree/icon_NonRequiredBone_16x"), Icon16x16));
-		Set("SkeletonTree.NormalFont", FTextBlockStyle(NormalText)
-			.SetColorAndOpacity(FLinearColor(1.0f, 1.0f, 1.0f)));
-		Set("SkeletonTree.BoldFont", FTextBlockStyle(NormalText)
-			.SetColorAndOpacity(FLinearColor(1.0f, 1.0f, 1.0f))
-			.SetFont(DEFAULT_FONT("Bold", 10)));
-
-		Set("SkeletonTree.HyperlinkSpinBox", FSpinBoxStyle(GetWidgetStyle<FSpinBoxStyle>("SpinBox"))
-			.SetTextPadding(FMargin(0))
-			.SetBackgroundBrush(BORDER_BRUSH("Old/HyperlinkDotted", FMargin(0, 0, 0, 3 / 16.0f), FSlateColor::UseSubduedForeground()))
-			.SetHoveredBackgroundBrush(FSlateNoResource())
-			.SetInactiveFillBrush(FSlateNoResource())
-			.SetActiveFillBrush(FSlateNoResource())
-			.SetForegroundColor(FSlateColor::UseSubduedForeground())
+		Set("SkeletonTree.NormalFont", FTextBlockStyle(NormalText));
+		Set("SkeletonTree.ItalicFont", FTextBlockStyle(NormalText)
+			.SetFont(DEFAULT_FONT("Italic", 10)));
+
+		Set("SkeletonTree.HyperlinkSpinBox", FSpinBoxStyle()
+			.SetBackgroundBrush(       FSlateRoundedBoxBrush(FStyleColors::Input, 3.f, FStyleColors::Transparent, 1.0f))
+			.SetHoveredBackgroundBrush(FSlateRoundedBoxBrush(FStyleColors::Input, 3.f, FStyleColors::Transparent, 1.0f))
+			.SetActiveFillBrush(       FSlateRoundedBoxBrush(FStyleColors::Hover, 3.f, FStyleColors::Transparent, 1.0f))
+			.SetInactiveFillBrush(     FSlateRoundedBoxBrush(FStyleColors::Secondary, 3.f, FStyleColors::Transparent, 1.0f))
 			.SetArrowsImage(FSlateNoResource())
-			);
+			.SetForegroundColor(FStyleColors::ForegroundHover)
+			.SetTextPadding(FMargin(8.f, 4.f, 8.f, 3.f))
+			.SetArrowsImage(FSlateNoResource())
+		);
+
+
+		const FButtonStyle BorderlessButton = FButtonStyle(GetWidgetStyle<FButtonStyle>("SimpleButton"))
+			.SetNormalPadding(0.f)
+			.SetPressedPadding(0.f);
+
+		Set("SkeletonTree.RetargetingComboButton", FComboButtonStyle(GetWidgetStyle<FComboButtonStyle>("SimpleComboButton"))
+			.SetButtonStyle(BorderlessButton)
+			.SetDownArrowPadding(FMargin(2.0f, 0.0f, 0.0f, 0.0f)));
 
 		Set("SkeletonTree.BlendProfile", new IMAGE_BRUSH(TEXT("Icons/icon_Persona_NewBlendSpace_16x"), Icon16x16));
 		Set("SkeletonTree.InlineEditorShadowTop", new IMAGE_BRUSH(TEXT("Common/ScrollBoxShadowTop"), FVector2D(64, 8)));
diff --git a/Engine/Source/Editor/PhysicsAssetEditor/Private/PhysicsAssetEditor.cpp b/Engine/Source/Editor/PhysicsAssetEditor/Private/PhysicsAssetEditor.cpp
--- a/Engine/Source/Editor/PhysicsAssetEditor/Private/PhysicsAssetEditor.cpp
+++ b/Engine/Source/Editor/PhysicsAssetEditor/Private/PhysicsAssetEditor.cpp
@@ -115,6 +115,18 @@ FPhysicsAssetEditor::~FPhysicsAssetEditor()
 	GEditor->GetEditorSubsystem<UImportSubsystem>()->OnAssetReimport.Remove(OnAssetReimportDelegateHandle);
 }
 
+static void FillAddShapeMenu(FMenuBuilder& InSubMenuBuilder)
+{
+	const FPhysicsAssetEditorCommands& PhysicsAssetEditorCommands = FPhysicsAssetEditorCommands::Get();
+
+	InSubMenuBuilder.BeginSection("ShapeTypeHeader", LOCTEXT("ShapeTypeHeader", "Shape Type"));
+	InSubMenuBuilder.AddMenuEntry( PhysicsAssetEditorCommands.AddBox );
+	InSubMenuBuilder.AddMenuEntry( PhysicsAssetEditorCommands.AddSphere );
+	InSubMenuBuilder.AddMenuEntry( PhysicsAssetEditorCommands.AddSphyl );
+	InSubMenuBuilder.AddMenuEntry( PhysicsAssetEditorCommands.AddTaperedCapsule );
+	InSubMenuBuilder.EndSection();
+}
+
 void FPhysicsAssetEditor::InitPhysicsAssetEditor(const EToolkitMode::Type Mode, const TSharedPtr< class IToolkitHost >& InitToolkitHost, UPhysicsAsset* ObjectToEdit)
 {
 	SelectedSimulation = false;
@@ -149,6 +161,7 @@ void FPhysicsAssetEditor::InitPhysicsAssetEditor(const EToolkitMode::Type Mode,
 	SkeletonTreeArgs.Extenders = MakeShared<FExtender>();
 	SkeletonTreeArgs.Extenders->AddMenuExtension("FilterOptions", EExtensionHook::After, GetToolkitCommands(), FMenuExtensionDelegate::CreateSP(this, &FPhysicsAssetEditor::HandleExtendFilterMenu));
 	SkeletonTreeArgs.Extenders->AddMenuExtension("SkeletonTreeContextMenu", EExtensionHook::After, GetToolkitCommands(), FMenuExtensionDelegate::CreateSP(this, &FPhysicsAssetEditor::HandleExtendContextMenu));
+	SkeletonTreeArgs.Extenders->AddMenuExtension("CreateNew", EExtensionHook::After, GetToolkitCommands(), FMenuExtensionDelegate::CreateStatic( &FillAddShapeMenu));
 	SkeletonTreeArgs.Builder = SkeletonTreeBuilder = MakeShared<FPhysicsAssetEditorSkeletonTreeBuilder>(ObjectToEdit);
 	SkeletonTreeArgs.ContextName = GetToolkitFName();
 
@@ -1156,18 +1169,6 @@ void FPhysicsAssetEditor::Mirror()
 	RefreshPreviewViewport();
 }
 
-static void FillAddShapeMenu(FMenuBuilder& InSubMenuBuilder)
-{
-	const FPhysicsAssetEditorCommands& PhysicsAssetEditorCommands = FPhysicsAssetEditorCommands::Get();
-
-	InSubMenuBuilder.BeginSection("ShapeTypeHeader", LOCTEXT("ShapeTypeHeader", "Shape Type"));
-	InSubMenuBuilder.AddMenuEntry( PhysicsAssetEditorCommands.AddBox );
-	InSubMenuBuilder.AddMenuEntry( PhysicsAssetEditorCommands.AddSphere );
-	InSubMenuBuilder.AddMenuEntry( PhysicsAssetEditorCommands.AddSphyl );
-	InSubMenuBuilder.AddMenuEntry( PhysicsAssetEditorCommands.AddTaperedCapsule );
-	InSubMenuBuilder.EndSection();
-}
-
 void FPhysicsAssetEditor::BuildMenuWidgetBody(FMenuBuilder& InMenuBuilder)
 {
 	InMenuBuilder.PushCommandList(GetToolkitCommands());
diff --git a/Engine/Source/Editor/PhysicsAssetEditor/Private/SkeletonTreePhysicsBodyItem.cpp b/Engine/Source/Editor/PhysicsAssetEditor/Private/SkeletonTreePhysicsBodyItem.cpp
--- a/Engine/Source/Editor/PhysicsAssetEditor/Private/SkeletonTreePhysicsBodyItem.cpp
+++ b/Engine/Source/Editor/PhysicsAssetEditor/Private/SkeletonTreePhysicsBodyItem.cpp
@@ -31,7 +31,7 @@ void FSkeletonTreePhysicsBodyItem::GenerateWidgetForNameColumn( TSharedPtr< SHor
 	];
 }
 
-TSharedRef< SWidget > FSkeletonTreePhysicsBodyItem::GenerateWidgetForDataColumn(const FName& DataColumnName)
+TSharedRef< SWidget > FSkeletonTreePhysicsBodyItem::GenerateWidgetForDataColumn(const FName& DataColumnName, FIsSelected InIsSelected)
 {
 	return SNullWidget::NullWidget;
 }
diff --git a/Engine/Source/Editor/PhysicsAssetEditor/Private/SkeletonTreePhysicsBodyItem.h b/Engine/Source/Editor/PhysicsAssetEditor/Private/SkeletonTreePhysicsBodyItem.h
--- a/Engine/Source/Editor/PhysicsAssetEditor/Private/SkeletonTreePhysicsBodyItem.h
+++ b/Engine/Source/Editor/PhysicsAssetEditor/Private/SkeletonTreePhysicsBodyItem.h
@@ -27,7 +27,7 @@ public:
 
 	/** ISkeletonTreeItem interface */
 	virtual void GenerateWidgetForNameColumn(TSharedPtr< SHorizontalBox > Box, const TAttribute<FText>& FilterText, FIsSelected InIsSelected) override;
-	virtual TSharedRef< SWidget > GenerateWidgetForDataColumn(const FName& DataColumnName) override;	
+	virtual TSharedRef< SWidget > GenerateWidgetForDataColumn(const FName& DataColumnName, FIsSelected InIsSelected) override;	
 	virtual FName GetRowItemName() const override { return BoneName; }
 	virtual UObject* GetObject() const override { return BodySetup; }
 
diff --git a/Engine/Source/Editor/PhysicsAssetEditor/Private/SkeletonTreePhysicsConstraintItem.cpp b/Engine/Source/Editor/PhysicsAssetEditor/Private/SkeletonTreePhysicsConstraintItem.cpp
--- a/Engine/Source/Editor/PhysicsAssetEditor/Private/SkeletonTreePhysicsConstraintItem.cpp
+++ b/Engine/Source/Editor/PhysicsAssetEditor/Private/SkeletonTreePhysicsConstraintItem.cpp
@@ -33,7 +33,7 @@ void FSkeletonTreePhysicsConstraintItem::GenerateWidgetForNameColumn( TSharedPtr
 	];
 }
 
-TSharedRef< SWidget > FSkeletonTreePhysicsConstraintItem::GenerateWidgetForDataColumn(const FName& DataColumnName)
+TSharedRef< SWidget > FSkeletonTreePhysicsConstraintItem::GenerateWidgetForDataColumn(const FName& DataColumnName, FIsSelected InIsSelected)
 {
 	return SNullWidget::NullWidget;
 }
diff --git a/Engine/Source/Editor/PhysicsAssetEditor/Private/SkeletonTreePhysicsConstraintItem.h b/Engine/Source/Editor/PhysicsAssetEditor/Private/SkeletonTreePhysicsConstraintItem.h
--- a/Engine/Source/Editor/PhysicsAssetEditor/Private/SkeletonTreePhysicsConstraintItem.h
+++ b/Engine/Source/Editor/PhysicsAssetEditor/Private/SkeletonTreePhysicsConstraintItem.h
@@ -29,7 +29,7 @@ public:
 
 	/** ISkeletonTreeItem interface */
 	virtual void GenerateWidgetForNameColumn(TSharedPtr< SHorizontalBox > Box, const TAttribute<FText>& FilterText, FIsSelected InIsSelected) override;
-	virtual TSharedRef< SWidget > GenerateWidgetForDataColumn(const FName& DataColumnName) override;	
+	virtual TSharedRef< SWidget > GenerateWidgetForDataColumn(const FName& DataColumnName, FIsSelected InIsSelected) override;	
 	virtual FName GetRowItemName() const override { return OtherBoneName; }
 	virtual UObject* GetObject() const override { return Constraint; }
 
diff --git a/Engine/Source/Editor/PhysicsAssetEditor/Private/SkeletonTreePhysicsShapeItem.cpp b/Engine/Source/Editor/PhysicsAssetEditor/Private/SkeletonTreePhysicsShapeItem.cpp
--- a/Engine/Source/Editor/PhysicsAssetEditor/Private/SkeletonTreePhysicsShapeItem.cpp
+++ b/Engine/Source/Editor/PhysicsAssetEditor/Private/SkeletonTreePhysicsShapeItem.cpp
@@ -74,7 +74,7 @@ void FSkeletonTreePhysicsShapeItem::GenerateWidgetForNameColumn( TSharedPtr< SHo
 	];
 }
 
-TSharedRef< SWidget > FSkeletonTreePhysicsShapeItem::GenerateWidgetForDataColumn(const FName& DataColumnName)
+TSharedRef< SWidget > FSkeletonTreePhysicsShapeItem::GenerateWidgetForDataColumn(const FName& DataColumnName, FIsSelected InIsSelected)
 {
 	return SNullWidget::NullWidget;
 }
diff --git a/Engine/Source/Editor/PhysicsAssetEditor/Private/SkeletonTreePhysicsShapeItem.h b/Engine/Source/Editor/PhysicsAssetEditor/Private/SkeletonTreePhysicsShapeItem.h
--- a/Engine/Source/Editor/PhysicsAssetEditor/Private/SkeletonTreePhysicsShapeItem.h
+++ b/Engine/Source/Editor/PhysicsAssetEditor/Private/SkeletonTreePhysicsShapeItem.h
@@ -21,7 +21,7 @@ public:
 
 	/** ISkeletonTreeItem interface */
 	virtual void GenerateWidgetForNameColumn(TSharedPtr< SHorizontalBox > Box, const TAttribute<FText>& FilterText, FIsSelected InIsSelected) override;
-	virtual TSharedRef< SWidget > GenerateWidgetForDataColumn(const FName& DataColumnName) override;	
+	virtual TSharedRef< SWidget > GenerateWidgetForDataColumn(const FName& DataColumnName, FIsSelected InIsSelected ) override;	
 	virtual FName GetRowItemName() const override { FString NameAsString = GetNameAsString(); return *NameAsString; }
 	virtual UObject* GetObject() const override { return BodySetup; }
 	virtual bool CanRenameItem() const override { return true; }
diff --git a/Engine/Source/Editor/SkeletonEditor/Private/SBlendProfilePicker.h b/Engine/Source/Editor/SkeletonEditor/Private/SBlendProfilePicker.h
--- a/Engine/Source/Editor/SkeletonEditor/Private/SBlendProfilePicker.h
+++ b/Engine/Source/Editor/SkeletonEditor/Private/SBlendProfilePicker.h
@@ -53,15 +53,19 @@ public:
 	/** Get the currently selected blend profile name */
 	FName GetSelectedBlendProfileName() const;
 
+	/** Create a New Blend profile withe the provided profile name */
+	void OnCreateNewProfileComitted(const FText& NewName, ETextCommit::Type CommitType);
+	
+	/* Deselect the current blend profile */
+	void OnClearSelection();
+
 private:
 
 	/** FEditorUndoClient interface */
 	virtual void PostUndo(bool bSuccess) override;
 	virtual void PostRedo(bool bSuccess) override;
 
-	void OnClearSelection();
 	void OnCreateNewProfile();
-	void OnCreateNewProfileComitted(const FText& NewName, ETextCommit::Type CommitType);
 
 	void OnProfileSelected(FName InBlendProfileName);
 	void OnProfileRemoved(FName InBlendProfileName);
diff --git a/Engine/Source/Editor/SkeletonEditor/Private/SSkeletonTree.cpp b/Engine/Source/Editor/SkeletonEditor/Private/SSkeletonTree.cpp
--- a/Engine/Source/Editor/SkeletonEditor/Private/SSkeletonTree.cpp
+++ b/Engine/Source/Editor/SkeletonEditor/Private/SSkeletonTree.cpp
@@ -39,8 +39,6 @@
 #include "ContentBrowserModule.h"
 #include "ComponentAssetBroker.h"
 
-
-
 #include "AnimPreviewInstance.h"
 
 #include "MeshUtilities.h"
@@ -172,8 +170,6 @@ void SSkeletonTree::Construct(const FArguments& InArgs, const TSharedRef<FEditab
 		RegisterOnSelectionChanged(FOnSkeletonTreeSelectionChanged::CreateRaw(PreviewScene.Pin().Get(), &IPersonaPreviewScene::HandleSkeletonTreeSelectionChanged));
 	}
 
-	InEditableSkeleton->RegisterOnSkeletonHierarchyChanged(USkeleton::FOnSkeletonHierarchyChanged::CreateSP(this, &SSkeletonTree::CreateTreeColumns));
-
 	PRAGMA_DISABLE_DEPRECATION_WARNINGS
 	if(InSkeletonTreeArgs.OnObjectSelected.IsBound())
 	{
@@ -213,48 +209,46 @@ void SSkeletonTree::Construct(const FArguments& InArgs, const TSharedRef<FEditab
 		+SOverlay::Slot()
 		[
 			SNew( SVerticalBox )
-		
 			+ SVerticalBox::Slot()
 			.AutoHeight()
-			.Padding(FMargin(0.0f, 0.0f, 0.0f, 2.0f))
+			.Padding(FMargin(0.f, 2.f))
 			[
 				SNew(SHorizontalBox)
 				+SHorizontalBox::Slot()
 				.AutoWidth()
 				.VAlign(VAlign_Center)
-				.Padding(0.0f, 0.0f, InSkeletonTreeArgs.bShowFilterMenu ? 2.0f : 0.0f, 0.0f)
+				.Padding(FMargin(6.f, 0.0))
 				[
 					SAssignNew(FilterComboButton, SComboButton)
-					.Visibility(InSkeletonTreeArgs.bShowFilterMenu ? EVisibility::Visible : EVisibility::Collapsed)
-					.ComboButtonStyle(FEditorStyle::Get(), "GenericFilters.ComboButtonStyle")
-					.ForegroundColor(FLinearColor::White)
-					.ContentPadding(0.0f)
-					.OnGetMenuContent( this, &SSkeletonTree::CreateFilterMenu )
-					.ToolTipText( this, &SSkeletonTree::GetFilterMenuTooltip )
-					.AddMetaData<FTagMetaData>(TEXT("SkelTree.Bones"))
+					.ButtonStyle(&FAppStyle::Get().GetWidgetStyle<FButtonStyle>("SecondaryButton"))
+					.ForegroundColor(FSlateColor::UseStyle())
+					.HasDownArrow(false)
+					.OnGetMenuContent( this, &SSkeletonTree::CreateNewMenu )
+					// .AddMetaData<FTagMetaData>(TEXT("SkelTree.Bones"))
 					.ButtonContent()
 					[
 						SNew(SHorizontalBox)
-						+SHorizontalBox::Slot()
+						+ SHorizontalBox::Slot()
 						.AutoWidth()
+						.HAlign(HAlign_Center)
 						.VAlign(VAlign_Center)
 						[
-							SNew(STextBlock)
-							.TextStyle(FEditorStyle::Get(), "GenericFilters.TextStyle")
-							.Font(FEditorStyle::Get().GetFontStyle("FontAwesome.9"))
-							.Text(FText::FromString(FString(TEXT("\xf0b0"))) /*fa-filter*/)
+							SNew(SImage)
+							.Image(FAppStyle::Get().GetBrush("Icons.Plus"))
+							.ColorAndOpacity(FSlateColor::UseForeground())
 						]
-						+SHorizontalBox::Slot()
-						.AutoWidth()
-						.Padding(2, 0, 0, 0)
+						+ SHorizontalBox::Slot()
+						.Padding(FMargin(7.f, 0.f, 0.0, 0.0))
 						.VAlign(VAlign_Center)
+						.AutoWidth()
 						[
-							SNew( STextBlock )
-							.TextStyle(FEditorStyle::Get(), "GenericFilters.TextStyle")
-							.Text( LOCTEXT("FilterMenuLabel", "Options") )
+							SNew(STextBlock)
+							.TextStyle(FAppStyle::Get(), "SmallButtonText")
+							.Text(LOCTEXT("NewButton", "ADD"))
 						]
 					]
 				]
+
 				+SHorizontalBox::Slot()
 				.FillWidth(1.0f)
 				[
@@ -264,6 +258,28 @@ void SSkeletonTree::Construct(const FArguments& InArgs, const TSharedRef<FEditab
 					.HintText( LOCTEXT( "SearchBoxHint", "Search Skeleton Tree...") )
 					.AddMetaData<FTagMetaData>(TEXT("SkelTree.Search"))
 				]
+
+				+SHorizontalBox::Slot()
+				.AutoWidth()
+				.Padding(FMargin(6.f, 0.0))
+				.VAlign(VAlign_Center)
+				[
+					SAssignNew(FilterComboButton, SComboButton)
+					.Visibility(InSkeletonTreeArgs.bShowFilterMenu ? EVisibility::Visible : EVisibility::Collapsed)
+					.ComboButtonStyle(&FAppStyle::Get().GetWidgetStyle<FComboButtonStyle>("SimpleComboButton"))
+					.ForegroundColor(FSlateColor::UseStyle())
+					.ContentPadding(2.0f)
+					.OnGetMenuContent( this, &SSkeletonTree::CreateFilterMenu )
+					.ToolTipText( this, &SSkeletonTree::GetFilterMenuTooltip )
+					.AddMetaData<FTagMetaData>(TEXT("SkelTree.Bones"))
+					.HasDownArrow(true)
+					.ButtonContent()
+					[
+						SNew(SImage)
+						.Image(FAppStyle::Get().GetBrush("Icons.Settings"))
+						.ColorAndOpacity(FSlateColor::UseForeground())
+					]
+				]
 			]
 
 			+ SVerticalBox::Slot()
@@ -278,32 +294,15 @@ void SSkeletonTree::Construct(const FArguments& InArgs, const TSharedRef<FEditab
 			[
 				SAssignNew(TreeHolder, SOverlay)
 			]
-
-			+ SVerticalBox::Slot()
-			.AutoHeight()
-			[
-				SNew( SHorizontalBox )
-
-				+ SHorizontalBox::Slot()
-				.Padding(0.0f, 0.0f, InSkeletonTreeArgs.bShowBlendProfiles ? 2.0f : 0.0f, 0.0f)
-				.VAlign(VAlign_Center)
-				.HAlign(HAlign_Left)
-				[
-					SNew(SHorizontalBox)
-					.Visibility(InSkeletonTreeArgs.bShowBlendProfiles ? EVisibility::Visible : EVisibility::Collapsed)
-					+ SHorizontalBox::Slot()
-					.Padding(0.0f, 0.0f, 2.0f, 0.0f)
-					.VAlign(VAlign_Center)
-					[
-						SAssignNew(BlendProfilePicker, SBlendProfilePicker, GetEditableSkeleton())
-						.Standalone(true)
-						.OnBlendProfileSelected(this, &SSkeletonTree::OnBlendProfileSelected)
-					]
-				]
-			]
 		]
 	];
 
+
+	SAssignNew(BlendProfilePicker, SBlendProfilePicker, GetEditableSkeleton())
+		.Standalone(true)
+		.OnBlendProfileSelected(this, &SSkeletonTree::OnBlendProfileSelected);
+
+
 	CreateTreeColumns();
 
 	SetInitialExpansionState();
@@ -480,6 +479,15 @@ void SSkeletonTree::BindCommands()
 		MenuActions.FocusCamera,
 		FExecuteAction::CreateSP(this, &SSkeletonTree::HandleFocusCamera));
 
+	CommandList.MapAction(
+		MenuActions.CreateBlendProfile,
+		FExecuteAction::CreateSP( this, &SSkeletonTree::OnCreateBlendProfile));
+
+	CommandList.MapAction(
+		MenuActions.DeleteCurrentBlendProfile,
+		FExecuteAction::CreateSP( this, &SSkeletonTree::OnDeleteCurrentBlendProfile));
+
+
 	PinnedCommands->BindCommandList(UICommandList.ToSharedRef());
 }
 
@@ -500,7 +508,7 @@ void SSkeletonTree::GetFilteredChildren(TSharedPtr<ISkeletonTreeItem> InInfo, TA
 struct FScopedSavedSelection
 {
 	FScopedSavedSelection(TSharedPtr<SSkeletonTree> InSkeletonTree)
-		: SkeletonTree(InSkeletonTree)
+	: SkeletonTree(InSkeletonTree)
 	{
 		// record selected items
 		if (SkeletonTree.IsValid() && InSkeletonTree->SkeletonTreeView.IsValid())
@@ -563,30 +571,45 @@ void SSkeletonTree::CreateTreeColumns()
 		
 		return  EVisibility::Collapsed; 
 	};
-		
-	TSharedRef<SHeaderRow> TreeHeaderRow = 
-		SNew(SHeaderRow)
-		.Visibility_Lambda(HeaderVisibilityLambda)
-		+ SHeaderRow::Column(ISkeletonTree::Columns::Name)
-		.DefaultLabel(LOCTEXT("SkeletonBoneNameLabel", "Name"))
-		.FillWidth(0.5f);
-
-	if (bShowingAdvancedOptions)
-	{
-		TreeHeaderRow->AddColumn(
-			SHeaderRow::Column(ISkeletonTree::Columns::Retargeting)
-			.DefaultLabel(LOCTEXT("SkeletonBoneTranslationRetargetingLabel", "Translation Retargeting"))
-			.FillWidth(0.25f)
-			);
-	}
 
-	if(BlendProfilePicker->GetSelectedBlendProfileName() != NAME_None)
-	{
-		TreeHeaderRow->AddColumn(
-			SHeaderRow::Column(ISkeletonTree::Columns::BlendProfile)
-			.DefaultLabel(LOCTEXT("BlendProfileLabel", "Blend Profile Scale"))
-			.FillWidth(0.25f));
-	}
+	TArray<FName> HiddenColumnsList;
+	HiddenColumnsList.Add(ISkeletonTree::Columns::Retargeting);
+	HiddenColumnsList.Add(ISkeletonTree::Columns::BlendProfile);
+
+	TSharedRef<SHeaderRow> TreeHeaderRow = 
+	SNew(SHeaderRow)
+	.CanSelectGeneratedColumn(true)
+	.HiddenColumnsList(HiddenColumnsList)
+
+	+ SHeaderRow::Column(ISkeletonTree::Columns::Name)
+	.ShouldGenerateWidget(true)
+	.DefaultLabel(LOCTEXT("SkeletonBoneNameLabel", "Name"))
+	.FillWidth(0.5f)
+
+	+ SHeaderRow::Column(ISkeletonTree::Columns::Retargeting)
+	.DefaultLabel(LOCTEXT("SkeletonBoneTranslationRetargetingLabel", "Translation Retargeting"))
+	.FillWidth(0.25f)
+
+	+ SHeaderRow::Column(ISkeletonTree::Columns::BlendProfile)
+	.DefaultLabel(LOCTEXT("BlendProfile", "Blend Profile"))
+	.FillWidth(0.25f)
+	.OnGetMenuContent(this, &SSkeletonTree::GetBlendProfileColumnMenuContent )
+	.HeaderContent()
+	[
+		SNew(SBox)
+		.HeightOverride(24)
+		.HAlign(HAlign_Left)
+		[
+			SAssignNew(BlendProfileHeader, SInlineEditableTextBlock)
+			.Text_Lambda([this] () -> FText
+			{
+				FName CurrentProfile = BlendProfilePicker->GetSelectedBlendProfileName();
+				return CurrentProfile != NAME_None ? FText::FromName(CurrentProfile) : LOCTEXT("NoBlendProfile", "No Blend");
+			})
+			.OnTextCommitted(BlendProfilePicker.ToSharedRef(), &SBlendProfilePicker::OnCreateNewProfileComitted)
+			.IsReadOnly(true)
+		]
+	];
 
 	{
 		FScopedSavedSelection ScopedSelection(SharedThis(this));
@@ -606,7 +629,7 @@ void SSkeletonTree::CreateTreeColumns()
 		(
 			TreeHeaderRow
 		);
-	
+
 		TreeHolder->ClearChildren();
 		TreeHolder->AddSlot()
 		[
@@ -659,7 +682,7 @@ void SSkeletonTree::ApplyFilter()
 		SetInitialExpansionState();
 	}
 
-	SkeletonTreeView->RequestTreeRefresh();
+	HandleTreeRefresh();
 }
 
 void SSkeletonTree::SetInitialExpansionState()
@@ -739,9 +762,9 @@ TSharedPtr< SWidget > SSkeletonTree::CreateContextMenu()
 					{
 						FUIAction RecursiveSetScales;
 						RecursiveSetScales.ExecuteAction = FExecuteAction::CreateSP(this, &SSkeletonTree::RecursiveSetBlendProfileScales, CurrentBlendScale);
-					
+
 						MenuBuilder.AddMenuEntry
-							(
+						(
 							FText::Format(LOCTEXT("RecursiveSetBlendScales_Label", "Recursively Set Blend Scales To {0}"), FText::AsNumber(CurrentBlendScale)),
 							LOCTEXT("RecursiveSetBlendScales_ToolTip", "Sets all child bones to use the same blend profile scale as the selected bone"),
 							FSlateIcon(),
@@ -762,39 +785,39 @@ TSharedPtr< SWidget > SSkeletonTree::CreateContextMenu()
 						FUIAction RecursiveRetargetingOrientAndScaleAction = FUIAction(FExecuteAction::CreateSP(this, &SSkeletonTree::SetBoneTranslationRetargetingModeRecursive, EBoneTranslationRetargetingMode::OrientAndScale));
 
 						MenuBuilder.AddMenuEntry
-							(LOCTEXT("SetTranslationRetargetingSkeletonChildrenAction", "Recursively Set Translation Retargeting Skeleton")
+						(LOCTEXT("SetTranslationRetargetingSkeletonChildrenAction", "Recursively Set Translation Retargeting Skeleton")
 							, LOCTEXT("BoneTranslationRetargetingSkeletonToolTip", "Use translation from Skeleton.")
 							, FSlateIcon()
 							, RecursiveRetargetingSkeletonAction
 							);
 
 						MenuBuilder.AddMenuEntry
-							(LOCTEXT("SetTranslationRetargetingAnimationChildrenAction", "Recursively Set Translation Retargeting Animation")
+						(LOCTEXT("SetTranslationRetargetingAnimationChildrenAction", "Recursively Set Translation Retargeting Animation")
 							, LOCTEXT("BoneTranslationRetargetingAnimationToolTip", "Use translation from animation.")
 							, FSlateIcon()
 							, RecursiveRetargetingAnimationAction
 							);
 
 						MenuBuilder.AddMenuEntry
-							(LOCTEXT("SetTranslationRetargetingAnimationScaledChildrenAction", "Recursively Set Translation Retargeting AnimationScaled")
+						(LOCTEXT("SetTranslationRetargetingAnimationScaledChildrenAction", "Recursively Set Translation Retargeting AnimationScaled")
 							, LOCTEXT("BoneTranslationRetargetingAnimationScaledToolTip", "Use translation from animation, scale length by Skeleton's proportions.")
 							, FSlateIcon()
 							, RecursiveRetargetingAnimationScaledAction
 							);
 
 						MenuBuilder.AddMenuEntry
-							(LOCTEXT("SetTranslationRetargetingAnimationRelativeChildrenAction", "Recursively Set Translation Retargeting AnimationRelative")
+						(LOCTEXT("SetTranslationRetargetingAnimationRelativeChildrenAction", "Recursively Set Translation Retargeting AnimationRelative")
 							, LOCTEXT("BoneTranslationRetargetingAnimationRelativeToolTip", "Use relative translation from animation similar to an additive animation.")
 							, FSlateIcon()
 							, RecursiveRetargetingAnimationRelativeAction
 							);
 
 						MenuBuilder.AddMenuEntry
-							(LOCTEXT("SetTranslationRetargetingOrientAndScaleChildrenAction", "Recursively Set Translation Retargeting OrientAndScale")
+						(LOCTEXT("SetTranslationRetargetingOrientAndScaleChildrenAction", "Recursively Set Translation Retargeting OrientAndScale")
 							, LOCTEXT("BoneTranslationRetargetingOrientAndScaleToolTip", "Orient And Scale Translation.")
 							, FSlateIcon()
 							, RecursiveRetargetingOrientAndScaleAction
-						);
+							);
 					}
 					MenuBuilder.EndSection();
 				}
@@ -919,10 +942,10 @@ void SSkeletonTree::FillVirtualBoneSubmenu(FMenuBuilder& MenuBuilder, TArray<TSh
 {
 	const bool bShowVirtualBones = false;
 	TSharedRef<SBoneTreeMenu> MenuContent = SNew(SBoneTreeMenu)
-		.bShowVirtualBones(false)
-		.Title(LOCTEXT("TargetBonePickerTitle", "Pick Target Bone..."))
-		.OnBoneSelectionChanged(this, &SSkeletonTree::OnVirtualTargetBonePicked, SourceBones)
-		.OnGetReferenceSkeleton(this, &SSkeletonTree::OnGetReferenceSkeleton);
+	.bShowVirtualBones(false)
+	.Title(LOCTEXT("TargetBonePickerTitle", "Pick Target Bone..."))
+	.OnBoneSelectionChanged(this, &SSkeletonTree::OnVirtualTargetBonePicked, SourceBones)
+	.OnGetReferenceSkeleton(this, &SSkeletonTree::OnGetReferenceSkeleton);
 	MenuBuilder.AddWidget(MenuContent, FText::GetEmpty(), true);
 
 	MenuContent->RegisterActiveTimer(0.0f, FWidgetActiveTimerDelegate::CreateLambda(
@@ -931,7 +954,7 @@ void SSkeletonTree::FillVirtualBoneSubmenu(FMenuBuilder& MenuBuilder, TArray<TSh
 			FSlateApplication::Get().SetKeyboardFocus(FilterTextBox);
 			return EActiveTimerReturnType::Stop;
 		}
-	));
+		));
 }
 
 void SSkeletonTree::OnVirtualTargetBonePicked(FName TargetBoneName, TArray<TSharedPtr<ISkeletonTreeItem>> SourceBones)
@@ -986,20 +1009,19 @@ void SSkeletonTree::OnVirtualTargetBonePicked(FName TargetBoneName, TArray<TShar
 			SkeletonTreeView->RequestScrollIntoView(LastItem);
 		}
 	}
-
 }
 
 void SSkeletonTree::CreateMenuForBoneReduction(FMenuBuilder& MenuBuilder, SSkeletonTree * Widget, int32 LODIndex, bool bIncludeSelected)
 {
 	MenuBuilder.AddMenuEntry
-		(FText::FromString(FString::Printf(TEXT("From LOD %d and below"), LODIndex))
+	(FText::FromString(FString::Printf(TEXT("From LOD %d and below"), LODIndex))
 		, FText::FromString(FString::Printf(TEXT("Remove Selected %s from current LOD %d and all lower LODs"), (bIncludeSelected) ? TEXT("bones") : TEXT("children"), LODIndex))
 		, FSlateIcon()
 		, FUIAction(FExecuteAction::CreateSP(Widget, &SSkeletonTree::RemoveFromLOD, LODIndex, bIncludeSelected, true))
 		);
 
 	MenuBuilder.AddMenuEntry
-		(FText::FromString(FString::Printf(TEXT("From LOD %d only"), LODIndex))
+	(FText::FromString(FString::Printf(TEXT("From LOD %d only"), LODIndex))
 		, FText::FromString(FString::Printf(TEXT("Remove selected %s from current LOD %d only"), (bIncludeSelected) ? TEXT("bones") : TEXT("children"), LODIndex))
 		, FSlateIcon()
 		, FUIAction(FExecuteAction::CreateSP(Widget, &SSkeletonTree::RemoveFromLOD, LODIndex, bIncludeSelected, false))
@@ -1077,7 +1099,6 @@ void SSkeletonTree::RemoveFromLOD(int32 LODIndex, bool bIncludeSelected, bool bI
 			int32 TotalLOD = PreviewMeshComponent->SkeletalMesh->GetLODNum();
 			IMeshUtilities& MeshUtilities = FModuleManager::Get().LoadModuleChecked<IMeshUtilities>("MeshUtilities");
 
-		
 			if (bIncludeBelowLODs)
 			{
 				for (int32 Index = LODIndex + 1; Index < TotalLOD; ++Index)
@@ -1318,7 +1339,7 @@ void SSkeletonTree::FillAttachAssetSubmenu(FMenuBuilder& MenuBuilder, const TSha
 	.HeightOverride(500)
 	[
 		ContentBrowserModule.Get().CreateAssetPicker(AssetPickerConfig)
-	];
+		];
 	MenuBuilder.AddWidget( MenuContent, FText::GetEmpty(), true);
 }
 
@@ -1493,7 +1514,7 @@ void SSkeletonTree::AttachAssets(const TSharedRef<ISkeletonTreeItem>& TargetItem
 	{
 		FName AttachToName = TargetItem->GetAttachName();
 		bool bAttachToMesh = TargetItem->IsOfType<FSkeletonTreeSocketItem>() &&
-			StaticCastSharedRef<FSkeletonTreeSocketItem>(TargetItem)->GetParentType() == ESocketParentType::Mesh;
+		StaticCastSharedRef<FSkeletonTreeSocketItem>(TargetItem)->GetParentType() == ESocketParentType::Mesh;
 		
 		GetEditableSkeletonInternal()->HandleAttachAssets(DroppedObjects, AttachToName, bAttachToMesh, GetPreviewScene());
 		CreateFromSkeleton();
@@ -1528,13 +1549,13 @@ void SSkeletonTree::SetTreeItemExpansionRecursive(TSharedPtr< ISkeletonTreeItem
 void SSkeletonTree::PostUndo(bool bSuccess)
 {
 	// Rebuild the tree view whenever we undo a change to the skeleton
-	CreateTreeColumns();
+	HandleTreeRefresh();
 }
 
 void SSkeletonTree::PostRedo(bool bSuccess)
 {
 	// Rebuild the tree view whenever we redo a change to the skeleton
-	CreateTreeColumns();
+	HandleTreeRefresh();
 }
 
 void SSkeletonTree::OnFilterTextChanged( const FText& SearchText )
@@ -1561,6 +1582,78 @@ void SSkeletonTree::HandlePackageReloaded(const EPackageReloadPhase InPackageRel
 	}
 }
 
+TSharedRef<SWidget> SSkeletonTree::GetBlendProfileColumnMenuContent()
+{
+	const FSkeletonTreeCommands& Actions = FSkeletonTreeCommands::Get();
+
+	FMenuBuilder MenuBuilder( true /*CloseAfterSelection*/, UICommandList );
+	GetBlendProfileMenu(MenuBuilder);
+	return MenuBuilder.MakeWidget();
+}
+
+void SSkeletonTree::GetBlendProfileMenu(FMenuBuilder& MenuBuilder)
+{
+	const FSkeletonTreeCommands& Actions = FSkeletonTreeCommands::Get();
+
+	MenuBuilder.BeginSection("BlendProfileActions", LOCTEXT("BlendProfiles", "Blend Profiles"));
+	{
+		for (UBlendProfile* Profile : GetEditableSkeletonInternal()->GetBlendProfiles())
+		{
+			MenuBuilder.AddMenuEntry(
+				FText::FromName(Profile->GetFName()),
+				LOCTEXT("SelectBlendProfileTooltip", "Select this profile for editing."),
+				FSlateIcon(),	
+				FUIAction(
+					FExecuteAction::CreateSP(BlendProfilePicker.ToSharedRef(), &SBlendProfilePicker::SetSelectedProfile, Profile, true),
+					FCanExecuteAction(),
+					FIsActionChecked::CreateSP(this, &SSkeletonTree::IsBlendProfileSelected, Profile->GetFName())
+				),
+				NAME_None,
+				EUserInterfaceActionType::RadioButton
+			);
+		}
+
+		MenuBuilder.AddMenuEntry(
+		LOCTEXT("Clear", "Clear"),
+		LOCTEXT("Clear_ToolTip", "Clear the selected blend profile."),
+		FSlateIcon(),
+		FUIAction(FExecuteAction::CreateSP(BlendProfilePicker.ToSharedRef(), &SBlendProfilePicker::OnClearSelection)));
+
+		MenuBuilder.AddSeparator();
+		MenuBuilder.AddMenuEntry(Actions.CreateBlendProfile);
+		if (BlendProfilePicker->GetSelectedBlendProfileName() != NAME_None)
+		{
+			MenuBuilder.AddSeparator();
+			MenuBuilder.AddMenuEntry(Actions.DeleteCurrentBlendProfile, NAME_None, FText::Format(LOCTEXT("DeleteBlendProfileLabel", "Remove {0}"), FText::FromName(BlendProfilePicker->GetSelectedBlendProfileName())));
+		}
+
+
+	}
+	MenuBuilder.EndSection();
+} 
+
+void SSkeletonTree::OnCreateBlendProfile()
+{
+	// Ensure the Blend Profile Column is Visible
+	BlendProfilePicker->OnClearSelection();
+	SkeletonTreeView->GetHeaderRow()->SetShowGeneratedColumn(ISkeletonTree::Columns::BlendProfile);
+
+	// Activate the Header Entry Box
+	BlendProfileHeader->SetReadOnly(false);
+	BlendProfileHeader->EnterEditingMode();
+}
+
+void SSkeletonTree::OnDeleteCurrentBlendProfile()
+{
+	GetEditableSkeletonInternal()->RemoveBlendProfile(BlendProfilePicker->GetSelectedBlendProfile());
+	BlendProfilePicker->OnClearSelection();
+}
+
+bool SSkeletonTree::IsBlendProfileSelected(FName ProfileName) const
+{
+	return BlendProfilePicker->GetSelectedBlendProfileName() == ProfileName;
+}
+
 void SSkeletonTree::Refresh()
 {
 	CreateFromSkeleton();
@@ -1652,6 +1745,32 @@ void SSkeletonTree::NotifyUser( FNotificationInfo& NotificationInfo )
 	}
 }
 
+TSharedRef< SWidget > SSkeletonTree::CreateNewMenu()
+{
+	const FSkeletonTreeCommands& Actions = FSkeletonTreeCommands::Get();
+	TArray<TSharedPtr<ISkeletonTreeItem>> SelectedItems = SkeletonTreeView->GetSelectedItems();
+
+	FMenuBuilder MenuBuilder( true /*CloseAfterSelection*/ , UICommandList, Extenders );
+
+	MenuBuilder.BeginSection("CreateNew", LOCTEXT("SkeletonCreateNew", "Create"));
+	{
+		MenuBuilder.AddMenuEntry(Actions.AddSocket);
+
+		MenuBuilder.AddSubMenu(LOCTEXT("AddVirtualBone", "Add Virtual Bone"),
+			LOCTEXT("AddVirtualBone_ToolTip", "Adds a virtual bone to the skeleton."),
+			FNewMenuDelegate::CreateSP(this, &SSkeletonTree::FillVirtualBoneSubmenu, SelectedItems));
+	}
+	MenuBuilder.EndSection();
+
+	MenuBuilder.BeginSection("Blend");
+	{
+		MenuBuilder.AddMenuEntry(Actions.CreateBlendProfile);
+	}
+	MenuBuilder.EndSection();
+
+	return MenuBuilder.MakeWidget();
+}
+
 TSharedRef< SWidget > SSkeletonTree::CreateFilterMenu()
 {
 	const FSkeletonTreeCommands& Actions = FSkeletonTreeCommands::Get();
@@ -1659,6 +1778,8 @@ TSharedRef< SWidget > SSkeletonTree::CreateFilterMenu()
 	const bool CloseAfterSelection = true;
 	FMenuBuilder MenuBuilder( CloseAfterSelection, UICommandList, Extenders );
 
+	GetBlendProfileMenu(MenuBuilder);
+
 	MenuBuilder.BeginSection("FilterOptions", LOCTEXT("OptionsMenuHeading", "Options"));
 	{
 		if(Builder->IsShowingBones() && bAllowSkeletonOperations)
@@ -1752,26 +1873,26 @@ FText SSkeletonTree::GetFilterMenuTooltip() const
 	{
 		switch ( BoneFilter )
 		{
-		case EBoneFilter::All:
+			case EBoneFilter::All:
 			FilterLabels.Add(LOCTEXT( "BoneFilterMenuAll", "Bones" ));
 			break;
 
-		case EBoneFilter::Mesh:
+			case EBoneFilter::Mesh:
 			FilterLabels.Add(LOCTEXT( "BoneFilterMenuMesh", "Mesh Bones" ));
 			break;
 
-		case EBoneFilter::LOD:
+			case EBoneFilter::LOD:
 			FilterLabels.Add(LOCTEXT("BoneFilterMenuLOD", "LOD Bones"));
 			break;
 
-		case EBoneFilter::Weighted:
+			case EBoneFilter::Weighted:
 			FilterLabels.Add(LOCTEXT( "BoneFilterMenuWeighted", "Weighted Bones" ));
 			break;
 
-		case EBoneFilter::None:
+			case EBoneFilter::None:
 			break;
 
-		default:
+			default:
 			// Unknown mode
 			check(false);
 			break;
@@ -1782,26 +1903,26 @@ FText SSkeletonTree::GetFilterMenuTooltip() const
 	{
 		switch (SocketFilter)
 		{
-		case ESocketFilter::Active:
+			case ESocketFilter::Active:
 			FilterLabels.Add(LOCTEXT("SocketFilterMenuActive", "Active Sockets"));
 			break;
 
-		case ESocketFilter::Mesh:
+			case ESocketFilter::Mesh:
 			FilterLabels.Add(LOCTEXT("SocketFilterMenuMesh", "Mesh Sockets"));
 			break;
 
-		case ESocketFilter::Skeleton:
+			case ESocketFilter::Skeleton:
 			FilterLabels.Add(LOCTEXT("SocketFilterMenuSkeleton", "Skeleton Sockets"));
 			break;
 
-		case ESocketFilter::All:
+			case ESocketFilter::All:
 			FilterLabels.Add(LOCTEXT("SocketFilterMenuAll", "All Sockets"));
 			break;
 
-		case ESocketFilter::None:
+			case ESocketFilter::None:
 			break;
 
-		default:
+			default:
 			// Unknown mode
 			check(false);
 			break;
@@ -1923,7 +2044,7 @@ void SSkeletonTree::DeleteVirtualBones(const TArray<TSharedPtr<FSkeletonTreeVirt
 void SSkeletonTree::OnChangeShowingAdvancedOptions()
 {
 	bShowingAdvancedOptions = !bShowingAdvancedOptions;
-	CreateTreeColumns();
+	HandleTreeRefresh();
 }
 
 bool SSkeletonTree::IsShowingAdvancedOptions() const
@@ -1943,7 +2064,12 @@ FName SSkeletonTree::GetSelectedBlendProfileName() const
 
 void SSkeletonTree::OnBlendProfileSelected(UBlendProfile* NewProfile)
 {
-	CreateTreeColumns();
+	SkeletonTreeView->GetHeaderRow()->RefreshColumns();
+	HandleTreeRefresh();
+
+	// When a new blend profile is created/selected - reaffirm that the header can't be edited.
+	// At this time, there is no re-naming of Blend Profiles
+	BlendProfileHeader->SetReadOnly(true);
 }
 
 void SSkeletonTree::RecursiveSetBlendProfileScales(float InScaleToSet)
@@ -1962,13 +2088,13 @@ void SSkeletonTree::RecursiveSetBlendProfileScales(float InScaleToSet)
 
 		GetEditableSkeletonInternal()->RecursiveSetBlendProfileScales(SelectedBlendProfile->GetFName(), BoneNames, InScaleToSet);
 
-		CreateTreeColumns();
+		HandleTreeRefresh();
 	}
 }
 
 void SSkeletonTree::HandleTreeRefresh()
 {
-	CreateFromSkeleton();
+	SkeletonTreeView->RequestTreeRefresh();
 }
 
 void SSkeletonTree::PostRenameSocket(UObject* InAttachedObject, const FName& InOldName, const FName& InNewName)
diff --git a/Engine/Source/Editor/SkeletonEditor/Private/SSkeletonTree.h b/Engine/Source/Editor/SkeletonEditor/Private/SSkeletonTree.h
--- a/Engine/Source/Editor/SkeletonEditor/Private/SSkeletonTree.h
+++ b/Engine/Source/Editor/SkeletonEditor/Private/SSkeletonTree.h
@@ -176,6 +176,9 @@ private:
 	/** Called to display context menu when right clicking on the widget */
 	TSharedPtr< SWidget > CreateContextMenu();
 
+	/** Called to display the add new menu */
+	TSharedRef< SWidget > CreateNewMenu();
+
 	/** Called to display the filter menu */
 	TSharedRef< SWidget > CreateFilterMenu();
 
@@ -320,6 +323,17 @@ private:
 	/** Handle package reloading (might be our skeleton) */
 	void HandlePackageReloaded(const EPackageReloadPhase InPackageReloadPhase, FPackageReloadedEvent* InPackageReloadedEvent);
 
+	/** Creates a new Blend Profile */
+	void OnCreateBlendProfile();
+
+	/** Removes the active Blend Profile */
+	void OnDeleteCurrentBlendProfile();
+
+	/** Generate Blend Profile Menu */
+	void GetBlendProfileMenu(FMenuBuilder& MenuBuilder);
+
+	TSharedRef<SWidget> GetBlendProfileColumnMenuContent();
+
 private:
 	/** Pointer back to the skeleton tree that owns us */
 	TWeakPtr<FEditableSkeleton> EditableSkeleton;
@@ -333,6 +347,9 @@ private:
 	/** The blend profile picker displaying the selected profile */
 	TSharedPtr<SBlendProfilePicker> BlendProfilePicker;
 
+	/** Blend Profile Header Label.  Also used to name new blend profiles */
+	TSharedPtr<SInlineEditableTextBlock> BlendProfileHeader;
+
 	/** Widget user to hold the skeleton tree */
 	TSharedPtr<SOverlay> TreeHolder;
 
@@ -386,6 +403,9 @@ private:
 	/** Add virtual bones to the skeleton tree */
 	void AddVirtualBones(const TArray<FVirtualBone>& VirtualBones);
 
+	/** Checks if the named profile is the currently selected/active one */
+	bool IsBlendProfileSelected(FName ProfileName) const;
+
 	/** The builder we use to construct the tree */
 	TSharedPtr<class ISkeletonTreeBuilder> Builder;
 
diff --git a/Engine/Source/Editor/SkeletonEditor/Private/SSkeletonTreeRow.cpp b/Engine/Source/Editor/SkeletonEditor/Private/SSkeletonTreeRow.cpp
--- a/Engine/Source/Editor/SkeletonEditor/Private/SSkeletonTreeRow.cpp
+++ b/Engine/Source/Editor/SkeletonEditor/Private/SSkeletonTreeRow.cpp
@@ -55,6 +55,7 @@ TSharedRef< SWidget > SSkeletonTreeRow::GenerateWidgetForColumn( const FName& Co
 			.AutoWidth()
 			[
 				SNew( SExpanderArrow, SharedThis(this) )
+				.ShouldDrawWires(true)
 			];
 
 		Item.Pin()->GenerateWidgetForNameColumn( RowBox, FilterText, FIsSelected::CreateSP(this, &STableRow::IsSelectedExclusively ) );
@@ -63,7 +64,7 @@ TSharedRef< SWidget > SSkeletonTreeRow::GenerateWidgetForColumn( const FName& Co
 	}
 	else
 	{
-		return Item.Pin()->GenerateWidgetForDataColumn(ColumnName);
+		return Item.Pin()->GenerateWidgetForDataColumn(ColumnName, FIsSelected::CreateSP(this, &STableRow::IsSelectedExclusively ));
 	}
 }
 
diff --git a/Engine/Source/Editor/SkeletonEditor/Private/SkeletonTreeAttachedAssetItem.cpp b/Engine/Source/Editor/SkeletonEditor/Private/SkeletonTreeAttachedAssetItem.cpp
--- a/Engine/Source/Editor/SkeletonEditor/Private/SkeletonTreeAttachedAssetItem.cpp
+++ b/Engine/Source/Editor/SkeletonEditor/Private/SkeletonTreeAttachedAssetItem.cpp
@@ -32,7 +32,7 @@ void FSkeletonTreeAttachedAssetItem::GenerateWidgetForNameColumn( TSharedPtr< SH
 	const FSlateBrush* IconBrush = ActorFactory ? FSlateIconFinder::FindIconBrushForClass(ActorFactory->GetDefaultActorClass(FAssetData())) : nullptr;
 	
 	Box->AddSlot()
-		.Padding(FMargin(0.0f, 1.0f))
+		.Padding(FMargin(0.0f, 2.0f))
 		.AutoWidth()
 		[
 			SNew( SImage )
@@ -64,7 +64,7 @@ void FSkeletonTreeAttachedAssetItem::GenerateWidgetForNameColumn( TSharedPtr< SH
 		];
 }
 
-TSharedRef< SWidget > FSkeletonTreeAttachedAssetItem::GenerateWidgetForDataColumn(const FName& DataColumnName)
+TSharedRef< SWidget > FSkeletonTreeAttachedAssetItem::GenerateWidgetForDataColumn(const FName& DataColumnName, FIsSelected InIsSelected)
 {
 	if(DataColumnName == ISkeletonTree::Columns::Retargeting)
 	{
@@ -81,7 +81,7 @@ TSharedRef< SWidget > FSkeletonTreeAttachedAssetItem::GenerateWidgetForDataColum
 				[
 					SNew(SImage)
 					.Image(this, &FSkeletonTreeAttachedAssetItem::OnGetAssetDisplayedButtonImage)
-					.ColorAndOpacity(FLinearColor::Black)
+					.ColorAndOpacity(FSlateColor::UseForeground())
 				]
 			];
 	}
diff --git a/Engine/Source/Editor/SkeletonEditor/Private/SkeletonTreeAttachedAssetItem.h b/Engine/Source/Editor/SkeletonEditor/Private/SkeletonTreeAttachedAssetItem.h
--- a/Engine/Source/Editor/SkeletonEditor/Private/SkeletonTreeAttachedAssetItem.h
+++ b/Engine/Source/Editor/SkeletonEditor/Private/SkeletonTreeAttachedAssetItem.h
@@ -19,7 +19,7 @@ public:
 
 	/** ISkeletonTreeItem interface */
 	virtual void GenerateWidgetForNameColumn(TSharedPtr< SHorizontalBox > Box, const TAttribute<FText>& FilterText, FIsSelected InIsSelected) override;
-	virtual TSharedRef< SWidget > GenerateWidgetForDataColumn(const FName& DataColumnName) override;
+	virtual TSharedRef< SWidget > GenerateWidgetForDataColumn(const FName& DataColumnName, FIsSelected InIsSelected) override;
 	virtual FName GetRowItemName() const { return Asset->GetFName(); }
 	virtual FName GetAttachName() const { return GetParentName(); }
 	virtual void OnItemDoubleClicked() override;
diff --git a/Engine/Source/Editor/SkeletonEditor/Private/SkeletonTreeBoneItem.cpp b/Engine/Source/Editor/SkeletonEditor/Private/SkeletonTreeBoneItem.cpp
--- a/Engine/Source/Editor/SkeletonEditor/Private/SkeletonTreeBoneItem.cpp
+++ b/Engine/Source/Editor/SkeletonEditor/Private/SkeletonTreeBoneItem.cpp
@@ -42,26 +42,32 @@ FSkeletonTreeBoneItem::FSkeletonTreeBoneItem(const FName& InBoneName, const TSha
 
 const FSlateBrush* FSkeletonTreeBoneItem::GetLODIcon() const
 {
-	if (bRequiredBone)
+	if (!bRequiredBone)
 	{
-		return FEditorStyle::GetBrush("SkeletonTree.LODBone");
+		return FAppStyle::Get().GetBrush("SkeletonTree.NonRequiredBone");
 	}
 
-	return FEditorStyle::GetBrush("SkeletonTree.NonRequiredBone");
+	else if (!bWeightedBone)
+	{
+		return FAppStyle::Get().GetBrush("SkeletonTree.BoneNonWeighted");
+	}
+
+	return FAppStyle::Get().GetBrush("SkeletonTree.Bone");
+
 }
 
 void FSkeletonTreeBoneItem::GenerateWidgetForNameColumn( TSharedPtr< SHorizontalBox > Box, const TAttribute<FText>& FilterText, FIsSelected InIsSelected )
 {
-	const FSlateBrush* LODIcon = FEditorStyle::GetBrush("SkeletonTree.LODBone");
+	const FSlateBrush* LODIcon = FEditorStyle::GetBrush("SkeletonTree.Bone");
 
 	Box->AddSlot()
 		.AutoWidth()
-		.Padding(FMargin(0.0f, 1.0f))
+		.Padding(FMargin(0.0f, 2.0f))
 		.VAlign(VAlign_Center)
 		.HAlign(HAlign_Center)
 		[
 			SNew(SImage)
-			.ColorAndOpacity(FSlateColor::UseForeground())
+			.ColorAndOpacity(this, &FSkeletonTreeBoneItem::GetBoneTextColor, InIsSelected)
 			.Image(this, &FSkeletonTreeBoneItem::GetLODIcon)
 		];
 
@@ -74,11 +80,11 @@ void FSkeletonTreeBoneItem::GenerateWidgetForNameColumn( TSharedPtr< SHorizontal
 	FText ToolTip = GetBoneToolTip();
 	Box->AddSlot()
 		.AutoWidth()
-		.Padding(2, 0, 0, 0)
+		.Padding(4, 0, 0, 0)
 		.VAlign(VAlign_Center)
 		[
 			SNew( STextBlock )
-			.ColorAndOpacity(this, &FSkeletonTreeBoneItem::GetBoneTextColor)
+			.ColorAndOpacity(this, &FSkeletonTreeBoneItem::GetBoneTextColor, InIsSelected)
 			.Text( FText::FromName(BoneName) )
 			.HighlightText( FilterText )
 			.Font(this, &FSkeletonTreeBoneItem::GetBoneTextFont)
@@ -86,7 +92,7 @@ void FSkeletonTreeBoneItem::GenerateWidgetForNameColumn( TSharedPtr< SHorizontal
 		];
 }
 
-TSharedRef< SWidget > FSkeletonTreeBoneItem::GenerateWidgetForDataColumn(const FName& DataColumnName)
+TSharedRef< SWidget > FSkeletonTreeBoneItem::GenerateWidgetForDataColumn(const FName& DataColumnName, FIsSelected InIsSelected)
 {
 	if(DataColumnName == ISkeletonTree::Columns::Retargeting)
 	{
@@ -94,55 +100,65 @@ TSharedRef< SWidget > FSkeletonTreeBoneItem::GenerateWidgetForDataColumn(const F
 			SNew(SHorizontalBox)
 			+SHorizontalBox::Slot()
 			.HAlign(HAlign_Left)
+			.Padding(0.0f)
 			[
 				SAssignNew(RetargetingComboButton, SComboButton)
-				.ButtonStyle(FEditorStyle::Get(), "ToggleButton")
-				.ForegroundColor(this, &FSkeletonTreeBoneItem::GetRetargetingComboButtonForegroundColor)
+            	.ComboButtonStyle( &FAppStyle::Get().GetWidgetStyle< FComboButtonStyle >("SkeletonTree.RetargetingComboButton"))
+				.ForegroundColor(this, &FSkeletonTreeBoneItem::GetBoneTextColor, InIsSelected)
 				.ContentPadding(0)
 				.OnGetMenuContent(this, &FSkeletonTreeBoneItem::CreateBoneTranslationRetargetingModeMenu)
 				.ToolTip(IDocumentation::Get()->CreateToolTip(
-				LOCTEXT("RetargetingToolTip", "Set bone translation retargeting mode"),
-				nullptr,
-				TEXT("Shared/Editors/Persona"),
-				TEXT("TranslationRetargeting")))
+					LOCTEXT("RetargetingToolTip", "Set bone translation retargeting mode"),
+					nullptr,
+					TEXT("Shared/Editors/Persona"),
+					TEXT("TranslationRetargeting")))
+				.VAlign(VAlign_Center)
 				.ButtonContent()
 				[
-					SNew(SHorizontalBox)
-					+SHorizontalBox::Slot()
-					.VAlign(VAlign_Center)
-					[
-						SNew(STextBlock)
-						.Text(this, &FSkeletonTreeBoneItem::GetTranslationRetargetingModeMenuTitle)
-					]
+					SNew(STextBlock)
+					.Text(this, &FSkeletonTreeBoneItem::GetTranslationRetargetingModeMenuTitle)
 				]
 			];
 	}
 	else if(DataColumnName == ISkeletonTree::Columns::BlendProfile)
 	{
-		bool bWritable = true;
-
-		UBlendProfile* CurrentProfile = GetSkeletonTree()->GetSelectedBlendProfile();
-
-		// We should never have this column if we don't have a profile
-		check(CurrentProfile);
-
 		return SNew(SBox)
 			.Padding(0.0f)
-			.HAlign(HAlign_Left)
+			.HAlign(HAlign_Fill)
+			.VAlign(VAlign_Center)
 			[
 				SNew(SSpinBox<float>)
-				.Style(&FEditorStyle::GetWidgetStyle<FSpinBoxStyle>("SkeletonTree.HyperlinkSpinBox"))
+				.Visibility(this, &FSkeletonTreeBoneItem::GetBoneBlendProfileVisibility)
+				.Style(&FAppStyle::Get(), "SkeletonTree.HyperlinkSpinBox")
+				.Font(FAppStyle::Get().GetFontStyle("SmallFont"))
 				.ContentPadding(0.0f)
+				.Delta(0.01f)
 				.MinValue(0.0f)
-				.MaxValue(1000.0f)
-				.Value(CurrentProfile->GetBoneBlendScale(BoneName))
+				.MaxValue(1.0f)
+				.Value(this, &FSkeletonTreeBoneItem::GetBoneBlendProfileScale)
 				.OnValueCommitted(this, &FSkeletonTreeBoneItem::OnBlendSliderCommitted)
+				.OnValueChanged(this, &FSkeletonTreeBoneItem::OnBlendSliderCommitted, ETextCommit::OnEnter)
 			];
 	}
 
 	return SNullWidget::NullWidget;
 }
 
+EVisibility FSkeletonTreeBoneItem::GetBoneBlendProfileVisibility() const
+{
+	return GetSkeletonTree()->GetSelectedBlendProfile() ? EVisibility::Visible : EVisibility::Collapsed;
+}
+
+float FSkeletonTreeBoneItem::GetBoneBlendProfileScale()	const 
+{
+	if (UBlendProfile* CurrentProfile = GetSkeletonTree()->GetSelectedBlendProfile())
+	{
+		return CurrentProfile->GetBoneBlendScale(BoneName);
+	}
+
+	return 0.0;
+}
+
 FSlateColor FSkeletonTreeBoneItem::GetRetargetingComboButtonForegroundColor() const
 {
 	static const FName InvertedForegroundName("InvertedForeground");
@@ -216,9 +232,9 @@ void FSkeletonTreeBoneItem::SetBoneBlendProfileScale(float NewScale, bool bRecur
 
 FSlateFontInfo FSkeletonTreeBoneItem::GetBoneTextFont() const
 {
-	if (bWeightedBone)
+	if (!bRequiredBone)
 	{
-		return FEditorStyle::GetWidgetStyle<FTextBlockStyle>("SkeletonTree.BoldFont").Font;
+		return FEditorStyle::GetWidgetStyle<FTextBlockStyle>("SkeletonTree.ItalicFont").Font;
 	}
 	else
 	{
@@ -254,19 +270,25 @@ void FSkeletonTreeBoneItem::EnableBoneProxyTick(bool bEnable)
 	BoneProxy->bIsTickable = bEnable;
 }
 
-FSlateColor FSkeletonTreeBoneItem::GetBoneTextColor() const
+FSlateColor FSkeletonTreeBoneItem::GetBoneTextColor(FIsSelected InIsSelected) const
 {
 	if (FilterResult == ESkeletonTreeFilterResult::ShownDescendant)
 	{
 		return FSlateColor(FLinearColor::Gray * 0.5f);
 	}
-	else if (bRequiredBone)
+
+	bool bIsSelected = InIsSelected.IsBound() ? InIsSelected.Execute() : false;
+	if (bIsSelected)
+	{
+		return FAppStyle::Get().GetSlateColor("Colors.ForegroundInverted");
+	}
+	else if (bRequiredBone && bWeightedBone)
 	{
-		return FSlateColor(FLinearColor::White);
+		return FSlateColor::UseForeground();
 	}
 	else
 	{
-		return FSlateColor(FLinearColor::Gray);
+		return FSlateColor::UseSubduedForeground();
 	}
 }
 
diff --git a/Engine/Source/Editor/SkeletonEditor/Private/SkeletonTreeBoneItem.h b/Engine/Source/Editor/SkeletonEditor/Private/SkeletonTreeBoneItem.h
--- a/Engine/Source/Editor/SkeletonEditor/Private/SkeletonTreeBoneItem.h
+++ b/Engine/Source/Editor/SkeletonEditor/Private/SkeletonTreeBoneItem.h
@@ -28,7 +28,7 @@ public:
 
 	/** ISkeletonTreeItem interface */
 	virtual void GenerateWidgetForNameColumn(TSharedPtr< SHorizontalBox > Box, const TAttribute<FText>& FilterText, FIsSelected InIsSelected) override;
-	virtual TSharedRef< SWidget > GenerateWidgetForDataColumn(const FName& DataColumnName) override;
+	virtual TSharedRef< SWidget > GenerateWidgetForDataColumn(const FName& DataColumnName, FIsSelected InIsSelected) override;
 	virtual FName GetRowItemName() const override { return BoneName; }
 	virtual void HandleDragEnter(const FDragDropEvent& DragDropEvent) override;
 	virtual void HandleDragLeave(const FDragDropEvent& DragDropEvent) override;
@@ -56,7 +56,7 @@ private:
 	FSlateFontInfo GetBoneTextFont() const;
 
 	/** Get the text color based on bone part of skeleton or part of mesh */
-	FSlateColor GetBoneTextColor() const;
+	FSlateColor GetBoneTextColor(FIsSelected InIsSelected) const;
 
 	/** Brush of the icon */
 	const FSlateBrush* GetLODIcon() const;
@@ -82,6 +82,12 @@ private:
 	/** Set current Blend Scale for this bone */
 	void SetBoneBlendProfileScale(float NewScale, bool bRecurse);
 
+	/** Get the current Blend Scale for this bone */
+	float GetBoneBlendProfileScale() const;
+
+	/** Used to hide the Bone Profile row widget if there is no current Blend Profile */
+	EVisibility GetBoneBlendProfileVisibility() const;
+
 private:
 	/** Bone proxy object */
 	UBoneProxy* BoneProxy;
diff --git a/Engine/Source/Editor/SkeletonEditor/Private/SkeletonTreeCommands.cpp b/Engine/Source/Editor/SkeletonEditor/Private/SkeletonTreeCommands.cpp
--- a/Engine/Source/Editor/SkeletonEditor/Private/SkeletonTreeCommands.cpp
+++ b/Engine/Source/Editor/SkeletonEditor/Private/SkeletonTreeCommands.cpp
@@ -35,6 +35,10 @@ void FSkeletonTreeCommands::RegisterCommands()
 	UI_COMMAND( HideSockets, "Hide Sockets", "Show no sockets", EUserInterfaceActionType::RadioButton, FInputChord() );
 
 	UI_COMMAND( FocusCamera, "Focus Camera", "Focus the camera on the current selection", EUserInterfaceActionType::Button, FInputChord(EKeys::F) );
+
+	UI_COMMAND( CreateBlendProfile, "Add Blend Profile", "Add a new blend profile for this skeleton", EUserInterfaceActionType::Button, FInputChord() );
+	UI_COMMAND( DeleteCurrentBlendProfile, "Delete the Current Blend Profile", "Delete the active blend profile for this skeleton", EUserInterfaceActionType::Button, FInputChord() );
+
 }
 
 #undef LOCTEXT_NAMESPACE
diff --git a/Engine/Source/Editor/SkeletonEditor/Private/SkeletonTreeCommands.h b/Engine/Source/Editor/SkeletonEditor/Private/SkeletonTreeCommands.h
--- a/Engine/Source/Editor/SkeletonEditor/Private/SkeletonTreeCommands.h
+++ b/Engine/Source/Editor/SkeletonEditor/Private/SkeletonTreeCommands.h
@@ -98,4 +98,10 @@ public:
 
 	/** Focus the camera on the current selection */
 	TSharedPtr< FUICommandInfo > FocusCamera;
+
+	/** Create a new blend profile */
+	TSharedPtr< FUICommandInfo > CreateBlendProfile;
+
+	/** Remove the currently active blend profile */
+	TSharedPtr< FUICommandInfo > DeleteCurrentBlendProfile;
 };
diff --git a/Engine/Source/Editor/SkeletonEditor/Private/SkeletonTreeSocketItem.cpp b/Engine/Source/Editor/SkeletonEditor/Private/SkeletonTreeSocketItem.cpp
--- a/Engine/Source/Editor/SkeletonEditor/Private/SkeletonTreeSocketItem.cpp
+++ b/Engine/Source/Editor/SkeletonEditor/Private/SkeletonTreeSocketItem.cpp
@@ -25,7 +25,7 @@ void FSkeletonTreeSocketItem::GenerateWidgetForNameColumn( TSharedPtr< SHorizont
 
 	Box->AddSlot()
 	.AutoWidth()
-	.Padding(FMargin(0.0f, 1.0f))
+	.Padding(FMargin(0.0f, 2.0f))
 	[
 		SNew( SImage )
 		.ColorAndOpacity(FSlateColor::UseForeground())
@@ -130,7 +130,7 @@ FSlateColor FSkeletonTreeSocketItem::GetTextColor() const
 	}
 }
 
-TSharedRef< SWidget > FSkeletonTreeSocketItem::GenerateWidgetForDataColumn(const FName& DataColumnName)
+TSharedRef< SWidget > FSkeletonTreeSocketItem::GenerateWidgetForDataColumn(const FName& DataColumnName, FIsSelected InIsSelected)
 {
 	return SNullWidget::NullWidget;
 }
diff --git a/Engine/Source/Editor/SkeletonEditor/Private/SkeletonTreeSocketItem.h b/Engine/Source/Editor/SkeletonEditor/Private/SkeletonTreeSocketItem.h
--- a/Engine/Source/Editor/SkeletonEditor/Private/SkeletonTreeSocketItem.h
+++ b/Engine/Source/Editor/SkeletonEditor/Private/SkeletonTreeSocketItem.h
@@ -32,7 +32,7 @@ public:
 
 	/** ISkeletonTreeItem interface */
 	virtual void GenerateWidgetForNameColumn(TSharedPtr< SHorizontalBox > Box, const TAttribute<FText>& FilterText, FIsSelected InIsSelected) override;
-	virtual TSharedRef< SWidget > GenerateWidgetForDataColumn(const FName& DataColumnName) override;
+	virtual TSharedRef< SWidget > GenerateWidgetForDataColumn(const FName& DataColumnName, FIsSelected InIsSelected) override;
 	virtual TSharedRef< SWidget > GenerateInlineEditWidget(const TAttribute<FText>& FilterText, FIsSelected InIsSelected) override;
 	virtual bool HasInlineEditor() const override { return GetDefault<UPersonaOptions>()->bUseInlineSocketEditor; }
 	virtual void ToggleInlineEditorExpansion() override { bInlineEditorExpanded = !bInlineEditorExpanded; }
diff --git a/Engine/Source/Editor/SkeletonEditor/Private/SkeletonTreeVirtualBoneItem.cpp b/Engine/Source/Editor/SkeletonEditor/Private/SkeletonTreeVirtualBoneItem.cpp
--- a/Engine/Source/Editor/SkeletonEditor/Private/SkeletonTreeVirtualBoneItem.cpp
+++ b/Engine/Source/Editor/SkeletonEditor/Private/SkeletonTreeVirtualBoneItem.cpp
@@ -36,16 +36,16 @@ EVisibility FSkeletonTreeVirtualBoneItem::GetLODIconVisibility() const
 
 void FSkeletonTreeVirtualBoneItem::GenerateWidgetForNameColumn(TSharedPtr< SHorizontalBox > Box, const TAttribute<FText>& FilterText, FIsSelected InIsSelected)
 {
-	const FSlateBrush* LODIcon = FEditorStyle::GetBrush("SkeletonTree.LODBone");
+	const FSlateBrush* LODIcon = FAppStyle::Get().GetBrush("SkeletonTree.Bone");
 
 	Box->AddSlot()
 		.AutoWidth()
-		.Padding(FMargin(0.0f, 1.0f))
+		.Padding(FMargin(0.0f, 2.0f))
 		.VAlign(VAlign_Center)
 		.HAlign(HAlign_Center)
 		[
 			SNew(SImage)
-			.ColorAndOpacity(FSlateColor::UseForeground())
+			.ColorAndOpacity(this, &FSkeletonTreeVirtualBoneItem::GetBoneTextColor, InIsSelected)
 			.Image(LODIcon)
 			.Visibility(this, &FSkeletonTreeVirtualBoneItem::GetLODIconVisibility)
 		];
@@ -69,7 +69,7 @@ void FSkeletonTreeVirtualBoneItem::GenerateWidgetForNameColumn(TSharedPtr< SHori
 
 	Box->AddSlot()
 		.AutoWidth()
-		.Padding(2, 2, 1, 0)
+		.Padding(4, 0, 0, 0)
 		[
 			SNew(STextBlock)
 			.ColorAndOpacity(this, &FSkeletonTreeVirtualBoneItem::GetBoneTextColor, InIsSelected)
@@ -79,13 +79,14 @@ void FSkeletonTreeVirtualBoneItem::GenerateWidgetForNameColumn(TSharedPtr< SHori
 		];
 
 	Box->AddSlot()
+		.Padding(4, 0, 0, 0)
 		.AutoWidth()
 		[
 			InlineWidget.ToSharedRef()
 		];
 }
 
-TSharedRef< SWidget > FSkeletonTreeVirtualBoneItem::GenerateWidgetForDataColumn(const FName& DataColumnName)
+TSharedRef< SWidget > FSkeletonTreeVirtualBoneItem::GenerateWidgetForDataColumn(const FName& DataColumnName, FIsSelected InIsSelected)
 {
 	return SNullWidget::NullWidget;
 }
diff --git a/Engine/Source/Editor/SkeletonEditor/Private/SkeletonTreeVirtualBoneItem.h b/Engine/Source/Editor/SkeletonEditor/Private/SkeletonTreeVirtualBoneItem.h
--- a/Engine/Source/Editor/SkeletonEditor/Private/SkeletonTreeVirtualBoneItem.h
+++ b/Engine/Source/Editor/SkeletonEditor/Private/SkeletonTreeVirtualBoneItem.h
@@ -26,7 +26,7 @@ public:
 
 	/** Builds the table row widget to display this info */
 	virtual void GenerateWidgetForNameColumn(TSharedPtr< SHorizontalBox > Box, const TAttribute<FText>& FilterText, FIsSelected InIsSelected) override;
-	virtual TSharedRef< SWidget > GenerateWidgetForDataColumn(const FName& DataColumnName) override;
+	virtual TSharedRef< SWidget > GenerateWidgetForDataColumn(const FName& DataColumnName, FIsSelected InIsSelected) override;
 	virtual FName GetRowItemName() const override { return BoneName; }
 	virtual bool CanRenameItem() const override { return true; }
 	virtual void RequestRename() override;
diff --git a/Engine/Source/Editor/SkeletonEditor/Public/ISkeletonTreeItem.h b/Engine/Source/Editor/SkeletonEditor/Public/ISkeletonTreeItem.h
--- a/Engine/Source/Editor/SkeletonEditor/Public/ISkeletonTreeItem.h
+++ b/Engine/Source/Editor/SkeletonEditor/Public/ISkeletonTreeItem.h
@@ -54,7 +54,7 @@ public:
 	virtual void GenerateWidgetForNameColumn(TSharedPtr< SHorizontalBox > Box, const TAttribute<FText>& FilterText, FIsSelected InIsSelected) = 0;
 
 	/** Builds the slate widget for the data column */
-	virtual TSharedRef< SWidget > GenerateWidgetForDataColumn(const FName& DataColumnName) = 0;
+	virtual TSharedRef< SWidget > GenerateWidgetForDataColumn(const FName& DataColumnName, FIsSelected InIsSelected) = 0;
 
 	/** Builds the slate widget for any inline data editing */
 	virtual TSharedRef< SWidget > GenerateInlineEditWidget(const TAttribute<FText>& FilterText, FIsSelected InIsSelected) = 0;
diff --git a/Engine/Source/Editor/SkeletonEditor/Public/SkeletonTreeItem.h b/Engine/Source/Editor/SkeletonEditor/Public/SkeletonTreeItem.h
--- a/Engine/Source/Editor/SkeletonEditor/Public/SkeletonTreeItem.h
+++ b/Engine/Source/Editor/SkeletonEditor/Public/SkeletonTreeItem.h
@@ -26,7 +26,7 @@ public:
 
 	/** ISkeletonTreeItem interface */
 	virtual TSharedRef<ITableRow> MakeTreeRowWidget(const TSharedRef<STableViewBase>& InOwnerTable, const TAttribute<FText>& InFilterText) override;
-	virtual TSharedRef< SWidget > GenerateWidgetForDataColumn(const FName& DataColumnName) override { return SNullWidget::NullWidget; }
+	virtual TSharedRef< SWidget > GenerateWidgetForDataColumn(const FName& DataColumnName, FIsSelected InIsSelected) override { return SNullWidget::NullWidget; }
 	virtual TSharedRef< SWidget > GenerateInlineEditWidget(const TAttribute<FText>& FilterText, FIsSelected InIsSelected) override { return SNullWidget::NullWidget; }
 	virtual bool HasInlineEditor() const override { return false; }
 	virtual void ToggleInlineEditorExpansion() override {}
