-Added Blend Profile mode selection to blend profile picker.
-Added option to blend profile picker for blend mask creation.
-Fixed issue w/ Blend profile picker sometimes showing in non-standalone mode outside of skeleton editor.
-Fixed issue w/ reset to default w/ blend profiles.
-Added blend mask mode to layered blend per bone node

diff --git a/Engine/Source/Editor/AnimGraph/Classes/AnimGraphNode_LayeredBoneBlend.h b/Engine/Source/Editor/AnimGraph/Classes/AnimGraphNode_LayeredBoneBlend.h
--- a/Engine/Source/Editor/AnimGraph/Classes/AnimGraphNode_LayeredBoneBlend.h
+++ b/Engine/Source/Editor/AnimGraph/Classes/AnimGraphNode_LayeredBoneBlend.h
@@ -29,6 +29,7 @@ class UAnimGraphNode_LayeredBoneBlend : public UAnimGraphNode_BlendListBase
 	virtual FLinearColor GetNodeTitleColor() const override;
 	virtual FText GetTooltipText() const override;
 	virtual FText GetNodeTitle(ENodeTitleType::Type TitleType) const override;
+	virtual void PostEditChangeProperty(struct FPropertyChangedEvent& PropertyChangedEvent) override;
 	// End of UEdGraphNode interface
 
 	// UK2Node interface
@@ -37,6 +38,11 @@ class UAnimGraphNode_LayeredBoneBlend : public UAnimGraphNode_BlendListBase
 
 	// UAnimGraphNode_Base interface
 	virtual FString GetNodeCategory() const override;
+	virtual void CustomizeDetails(IDetailLayoutBuilder& DetailBuilder) override;
+	virtual void PreloadRequiredAssets() override;
+	// End of UAnimGraphNode_Base interface
+
+
 	// Gives each visual node a chance to validate that they are still valid in the context of the compiled class, giving a last shot at error or warning generation after primary compilation is finished
 	virtual void ValidateAnimNodeDuringCompilation(class USkeleton* ForSkeleton, class FCompilerResultsLog& MessageLog) override;
 	// End of UAnimGraphNode_Base interface
diff --git a/Engine/Source/Editor/AnimGraph/Classes/AnimStateTransitionNode.h b/Engine/Source/Editor/AnimGraph/Classes/AnimStateTransitionNode.h
--- a/Engine/Source/Editor/AnimGraph/Classes/AnimStateTransitionNode.h
+++ b/Engine/Source/Editor/AnimGraph/Classes/AnimStateTransitionNode.h
@@ -49,7 +49,7 @@ class UAnimStateTransitionNode : public UAnimStateNodeBase
 	UCurveFloat* CustomBlendCurve;
 
 	// The blend profile to use to evaluate this transition per-bone
-	UPROPERTY(EditAnywhere, Category=Transition)
+	UPROPERTY(EditAnywhere, Category=Transition, meta=(UseAsBlendProfile=true))
 	UBlendProfile* BlendProfile;
 
 	// The blend profile mode to use. This determines how the per bone blend weights are calculated when transitioning.
diff --git a/Engine/Source/Editor/AnimGraph/Private/AnimGraphNode_LayeredBoneBlend.cpp b/Engine/Source/Editor/AnimGraph/Private/AnimGraphNode_LayeredBoneBlend.cpp
--- a/Engine/Source/Editor/AnimGraph/Private/AnimGraphNode_LayeredBoneBlend.cpp
+++ b/Engine/Source/Editor/AnimGraph/Private/AnimGraphNode_LayeredBoneBlend.cpp
@@ -7,6 +7,8 @@
 #include "GraphEditorActions.h"
 #include "ScopedTransaction.h"
 
+#include "DetailLayoutBuilder.h"
+#include "Kismet2/CompilerResultsLog.h"
 /////////////////////////////////////////////////////
 // UAnimGraphNode_LayeredBoneBlend
 
@@ -33,11 +35,59 @@ FText UAnimGraphNode_LayeredBoneBlend::GetNodeTitle(ENodeTitleType::Type TitleTy
 	return LOCTEXT("AnimGraphNode_LayeredBoneBlend_Title", "Layered blend per bone");
 }
 
+void UAnimGraphNode_LayeredBoneBlend::PostEditChangeProperty(struct FPropertyChangedEvent& PropertyChangedEvent)
+{
+	const FName PropertyName = (PropertyChangedEvent.Property ? PropertyChangedEvent.Property->GetFName() : NAME_None);
+
+	// Reconstruct node to show updates to PinFriendlyNames.
+	if (PropertyName == GET_MEMBER_NAME_STRING_CHECKED(FAnimNode_LayeredBoneBlend, BlendMode))
+	{
+		// If we  change blend modes, we need to resize our containers
+		Node.ValidateData();
+		FBlueprintEditorUtils::MarkBlueprintAsStructurallyModified(GetBlueprint());
+	}
+
+	Super::PostEditChangeProperty(PropertyChangedEvent);
+}
+
 FString UAnimGraphNode_LayeredBoneBlend::GetNodeCategory() const
 {
 	return TEXT("Blends");
 }
 
+void UAnimGraphNode_LayeredBoneBlend::CustomizeDetails(IDetailLayoutBuilder& DetailBuilder)
+{
+	TSharedRef<IPropertyHandle> NodeHandle = DetailBuilder.GetProperty(FName(TEXT("Node")), GetClass());
+
+	if (Node.BlendMode != ELayeredBoneBlendMode::BranchFilter)
+	{
+		DetailBuilder.HideProperty(NodeHandle->GetChildHandle(GET_MEMBER_NAME_CHECKED(FAnimNode_LayeredBoneBlend, LayerSetup)));
+	}
+
+	if (Node.BlendMode != ELayeredBoneBlendMode::BlendMask)
+	{
+		DetailBuilder.HideProperty(NodeHandle->GetChildHandle(GET_MEMBER_NAME_CHECKED(FAnimNode_LayeredBoneBlend, BlendMasks)));
+	}
+
+	Super::CustomizeDetails(DetailBuilder);
+}
+
+void UAnimGraphNode_LayeredBoneBlend::PreloadRequiredAssets()
+{
+	// Preload our blend profiles in case they haven't been loaded by the skeleton yet.
+	if (Node.BlendMode == ELayeredBoneBlendMode::BlendMask)
+	{
+		int32 NumBlendMasks = Node.BlendMasks.Num();
+		for (int32 MaskIndex = 0; MaskIndex < NumBlendMasks; ++MaskIndex)
+		{
+			UBlendProfile* BlendMask = Node.BlendMasks[MaskIndex];
+			PreloadObject(BlendMask);
+		}
+	}
+
+	Super::PreloadRequiredAssets();
+}
+
 void UAnimGraphNode_LayeredBoneBlend::AddPinToBlendByFilter()
 {
 	FScopedTransaction Transaction( LOCTEXT("AddPinToBlend", "AddPinToBlendByFilter") );
@@ -101,6 +151,34 @@ void UAnimGraphNode_LayeredBoneBlend::ValidateAnimNodeDuringCompilation(class US
 {
 	UAnimGraphNode_Base::ValidateAnimNodeDuringCompilation(ForSkeleton, MessageLog);
 
+	bool bCompilationError = false;
+	// Validate blend masks
+	if (Node.BlendMode == ELayeredBoneBlendMode::BlendMask)
+	{
+		int32 NumBlendMasks = Node.BlendMasks.Num();
+		for (int32 MaskIndex = 0; MaskIndex < NumBlendMasks; ++MaskIndex)
+		{
+			const UBlendProfile* BlendMask = Node.BlendMasks[MaskIndex];
+			if (BlendMask == nullptr)
+			{
+				MessageLog.Error(*FText::Format(LOCTEXT("LayeredBlendNullMask", "@@ has null BlendMask for Blend Pose {0}. "), FText::AsNumber(MaskIndex)).ToString(), this, BlendMask);
+				bCompilationError = true;
+				continue;
+			}
+			else if (BlendMask->BlendProfileMode != EBlendProfileMode::BlendMask)
+			{
+				MessageLog.Error(*FText::Format(LOCTEXT("LayeredBlendProfileModeError", "@@ is using a BlendProfile(@@) without a BlendMask mode for Blend Pose {0}. "), FText::AsNumber(MaskIndex)).ToString(), this, BlendMask);
+				bCompilationError = true;
+			}
+		}
+	}
+
+	// Don't rebuild the node's data if compilation failed. We may be attempting to do so with invalid data.
+	if (bCompilationError)
+	{
+		return;
+	}
+
 	// ensure to cache the data
  	if (Node.IsCacheInvalid(ForSkeleton))
  	{
diff --git a/Engine/Source/Editor/Persona/Private/BlendProfileCustomization.cpp b/Engine/Source/Editor/Persona/Private/BlendProfileCustomization.cpp
--- a/Engine/Source/Editor/Persona/Private/BlendProfileCustomization.cpp
+++ b/Engine/Source/Editor/Persona/Private/BlendProfileCustomization.cpp
@@ -14,6 +14,13 @@
 
 #define LOCTEXT_NAMESPACE "BlendProfileCustomization"
 
+namespace BlendProfileCustomizationNames
+{
+	
+	static const FName UseAsBlendProfileName = FName(TEXT("UseAsBlendProfile"));
+	static const FName UseAsBlendMaskName = FName(TEXT("UseAsBlendMask"));
+}
+
 TSharedRef<IPropertyTypeCustomization> FBlendProfileCustomization::MakeInstance()
 {
 	return MakeShareable(new FBlendProfileCustomization);
@@ -37,12 +44,30 @@ void FBlendProfileCustomization::CustomizeHeader(TSharedRef<class IPropertyHandl
 			InStructPropertyHandle->GetValue(PropertyValue);
 			UBlendProfile* CurrentProfile = Cast<UBlendProfile>(PropertyValue);
 
+			const bool bUseAsBlendMask = InStructPropertyHandle->GetBoolMetaData(BlendProfileCustomizationNames::UseAsBlendMaskName);
+			const bool bUseAsBlendProfile = InStructPropertyHandle->GetBoolMetaData(BlendProfileCustomizationNames::UseAsBlendProfileName);
+
+			// If no mode is defined, show both.
+			EBlendProfilePickerMode SupportedBlendProfileModes = (bUseAsBlendMask || bUseAsBlendProfile) ? EBlendProfilePickerMode(0) : EBlendProfilePickerMode::AllModes;
+			
+			if (bUseAsBlendProfile)
+			{
+				SupportedBlendProfileModes |= EBlendProfilePickerMode::BlendProfile;
+			}
+			if (bUseAsBlendMask)
+			{
+				SupportedBlendProfileModes |= EBlendProfilePickerMode::BlendMask;
+			}
+
+
 			FBlendProfilePickerArgs Args;
 			Args.bAllowNew = false;
-			Args.bAllowRemove = false;
+			Args.bAllowModify = false;
 			Args.bAllowClear = true;
 			Args.OnBlendProfileSelected = FOnBlendProfileSelected::CreateSP(this, &FBlendProfileCustomization::OnBlendProfileChanged, PropertyPtr);
 			Args.InitialProfile = CurrentProfile;
+			Args.SupportedBlendProfileModes = SupportedBlendProfileModes;
+			Args.PropertyHandle = PropertyPtr;
 
 			ISkeletonEditorModule& SkeletonEditorModule = FModuleManager::Get().LoadModuleChecked<ISkeletonEditorModule>("SkeletonEditor");
 			ValueCustomWidget = SkeletonEditorModule.CreateBlendProfilePicker(TargetSkeleton, Args);
@@ -53,7 +78,9 @@ void FBlendProfileCustomization::CustomizeHeader(TSharedRef<class IPropertyHandl
 		[
 			InStructPropertyHandle->CreatePropertyNameWidget()
 		]
-	.ValueContent()
+		.ValueContent()
+		.MinDesiredWidth(125.f)
+		.MaxDesiredWidth(400.f) //Slightly wider since expected names are a bit longer if users use BlendProfileModes as suffix
 		[
 			// If we can't find a skeleton, we can't use custom SWidget. Default to regular one.
 			ValueCustomWidget != SNullWidget::NullWidget ? ValueCustomWidget : InStructPropertyHandle->CreatePropertyValueWidget()
diff --git a/Engine/Source/Editor/SkeletonEditor/Private/EditableSkeleton.cpp b/Engine/Source/Editor/SkeletonEditor/Private/EditableSkeleton.cpp
--- a/Engine/Source/Editor/SkeletonEditor/Private/EditableSkeleton.cpp
+++ b/Engine/Source/Editor/SkeletonEditor/Private/EditableSkeleton.cpp
@@ -183,6 +183,30 @@ void FEditableSkeleton::SetBlendProfileScale(const FName& InBlendProfileName, co
 	}
 }
 
+void FEditableSkeleton::SetBlendProfileMode(FName InBlendProfileName, EBlendProfileMode ProfileMode)
+{
+	UBlendProfile* BlendProfile = GetBlendProfile(InBlendProfileName);
+	if (BlendProfile)
+	{
+		FScopedTransaction Transaction(LOCTEXT("SetBlendProfileMode", "Set Blend Profile Mode"));
+
+		const bool bWasBlendMask = BlendProfile->IsBlendMask();
+		BlendProfile->SetFlags(RF_Transactional);
+		BlendProfile->Modify();
+		BlendProfile->BlendProfileMode = ProfileMode;
+		const bool bIsBlendMask = BlendProfile->IsBlendMask();
+		// Re-set entry indices from the end to properly handle different default values when changing from/to blend mask
+		if (bWasBlendMask != bIsBlendMask)
+		{
+			for (int32 EntryIndex = BlendProfile->GetNumBlendEntries() - 1; EntryIndex >= 0; --EntryIndex)
+			{
+				const FBlendProfileBoneEntry& Entry = BlendProfile->ProfileEntries[EntryIndex];
+				BlendProfile->SetBoneBlendScale(Entry.BoneReference.BoneName, Entry.BlendScale);
+			}
+		}
+	}
+}
+
 void FEditableSkeleton::AddReferencedObjects(FReferenceCollector& Collector)
 {
 	Collector.AddReferencedObject(Skeleton);
@@ -1061,7 +1085,9 @@ TSharedRef<SWidget> FEditableSkeleton::CreateBlendProfilePicker(const FBlendProf
 		.OnBlendProfileSelected(InArgs.OnBlendProfileSelected)
 		.AllowNew(InArgs.bAllowNew)
 		.AllowClear(InArgs.bAllowClear)
-		.AllowRemove(InArgs.bAllowRemove);
+		.AllowModify(InArgs.bAllowModify)
+		.SupportedBlendProfileModes(InArgs.SupportedBlendProfileModes)
+		.PropertyHandle(InArgs.PropertyHandle);
 
 	BlendProfilePickers.Add(BlendProfilePicker);
 	return BlendProfilePicker;
diff --git a/Engine/Source/Editor/SkeletonEditor/Private/EditableSkeleton.h b/Engine/Source/Editor/SkeletonEditor/Private/EditableSkeleton.h
--- a/Engine/Source/Editor/SkeletonEditor/Private/EditableSkeleton.h
+++ b/Engine/Source/Editor/SkeletonEditor/Private/EditableSkeleton.h
@@ -34,6 +34,7 @@ public:
 	virtual class UBlendProfile* CreateNewBlendProfile(const FName& InBlendProfileName) override;
 	virtual void RemoveBlendProfile(UBlendProfile* InBlendProfile) override;
 	virtual void SetBlendProfileScale(const FName& InBlendProfileName, const FName& InBoneName, float InNewScale, bool bInRecurse) override;
+	virtual void SetBlendProfileMode(FName InBlendProfileName, EBlendProfileMode ProfileMode);
 	virtual USkeletalMeshSocket* DuplicateSocket(const FSelectedSocketInfo& SocketInfoToDuplicate, const FName& NewParentBoneName, USkeletalMesh* InSkeletalMesh) override;
 	virtual int32 ValidatePreviewAttachedObjects() override;
 	virtual int32 DeleteAnimNotifies(const TArray<FName>& InSelectedNotifyNames) override;
diff --git a/Engine/Source/Editor/SkeletonEditor/Private/SBlendProfilePicker.cpp b/Engine/Source/Editor/SkeletonEditor/Private/SBlendProfilePicker.cpp
--- a/Engine/Source/Editor/SkeletonEditor/Private/SBlendProfilePicker.cpp
+++ b/Engine/Source/Editor/SkeletonEditor/Private/SBlendProfilePicker.cpp
@@ -24,19 +24,33 @@
 
 #define LOCTEXT_NAMESPACE "BlendProfilePicker"
 
+namespace BlendProfilePickerNames
+{
+	static const FText BlendProfileModeName = LOCTEXT("Blend Profile Mode Name", "Blend Profile");
+	static const FText BlendMaskModeName = LOCTEXT("Blend Mask Mode Name", "Blend Mask");
+
+	const FText& GetNameForMode(EBlendProfileMode InMode)
+	{
+		return (InMode == EBlendProfileMode::BlendMask) ? BlendProfilePickerNames::BlendMaskModeName : BlendProfilePickerNames::BlendProfileModeName;
+	}
+}
 
 class SBlendProfileMenuEntry : public SCompoundWidget
 {
 public:
 
+	DECLARE_DELEGATE_OneParam(FBlendProfileModeChanged, EBlendProfileMode);
+
 	SLATE_BEGIN_ARGS(SBlendProfileMenuEntry){}
 		SLATE_ARGUMENT( FText, LabelOverride )
 		/** Called to when an entry is clicked */
 		SLATE_EVENT( FExecuteAction, OnOpenClickedDelegate )
 		/** Called to when the button remove an entry is clicked */
 		SLATE_EVENT( FExecuteAction, OnRemoveClickedDelegate )
+		SLATE_EVENT(FBlendProfileModeChanged, OnProfileModeChangedDelegate)
 		/** Whether to show the remove button */
-		SLATE_ARGUMENT(bool, AllowRemove)
+		SLATE_ARGUMENT(bool, AllowModify)
+		SLATE_ARGUMENT(TWeakObjectPtr<UBlendProfile>, BlendProfile)
 	SLATE_END_ARGS()
 
 	void Construct( const FArguments& InArgs )
@@ -44,64 +58,157 @@ public:
 		const FText DisplayName = InArgs._LabelOverride;
 		OnOpenClickedDelegate = InArgs._OnOpenClickedDelegate;
 		OnRemoveClickedDelegate = InArgs._OnRemoveClickedDelegate;
+		OnProfileModeChangedDelegate = InArgs._OnProfileModeChangedDelegate;
+		BlendProfile = InArgs._BlendProfile;
 
 		FSlateFontInfo MenuEntryFont = FEditorStyle::GetFontStyle( "Menu.Label.Font" );
 
-		TSharedPtr<SOverlay> Overlay;
+		TSharedPtr<SHorizontalBox> HorizontalBox;
+
+		EBlendProfileMode BlendProfileMode = GetProfileMode();
 
 		ChildSlot
 		[
 			SNew(SButton)
 			.ButtonStyle( FEditorStyle::Get(), "Menu.Button" )
 			.ForegroundColor( TAttribute<FSlateColor>::Create( TAttribute<FSlateColor>::FGetter::CreateRaw( this, &SBlendProfileMenuEntry::InvertOnHover ) ) )
-			.Text(DisplayName)
 			.ToolTipText(LOCTEXT("OpenBlendProfileToolTip", "Select this profile for editing."))
 			.OnClicked(this, &SBlendProfileMenuEntry::OnOpen)
 			.VAlign(VAlign_Center)
 			.HAlign(HAlign_Fill)
 			.ContentPadding( FMargin(4.0, 2.0) )
 			[
-				SAssignNew(Overlay, SOverlay)
-
+				SNew(SOverlay)
 				+SOverlay::Slot()
-				.Padding( FMargin( 12.0, 0.0 ) )
 				.VAlign(VAlign_Center)
-				.HAlign(HAlign_Left)
+				.HAlign(HAlign_Fill)
 				[
-					SNew( STextBlock )
-					.Font( MenuEntryFont )
-					.ColorAndOpacity( TAttribute<FSlateColor>::Create( TAttribute<FSlateColor>::FGetter::CreateRaw( this, &SBlendProfileMenuEntry::InvertOnHover ) ) )
-					.Text( DisplayName )
+					SAssignNew(HorizontalBox, SHorizontalBox)
+					+SHorizontalBox::Slot()
+					.Padding( FMargin( 12.0, 0.0 ) )
+					.VAlign(VAlign_Center)
+					.HAlign(HAlign_Fill)
+					[
+						SNew(STextBlock)
+						.Font(MenuEntryFont)
+						.ColorAndOpacity(TAttribute<FSlateColor>::Create(TAttribute<FSlateColor>::FGetter::CreateRaw(this, &SBlendProfileMenuEntry::InvertOnHover)))
+						.Text(DisplayName)
+					]
 				]
 			]
 		];
 
-		if(InArgs._AllowRemove)
+		if(InArgs._AllowModify)
 		{
-			Overlay->AddSlot()
-				.Padding(FMargin(0.0, 0.0))
+			HorizontalBox->AddSlot()
+				.Padding(FMargin(8.0, 0.0))
 				.VAlign(VAlign_Center)
 				.HAlign(HAlign_Right)
+				.AutoWidth()
 				[
-					SNew(SButton)
-					.ContentPadding(FMargin(4.0, 0.0))
-					.ButtonStyle(FEditorStyle::Get(), "Docking.Tab.CloseButton")
-					.ToolTipText(FText::Format(LOCTEXT("RemoveBlendProfileToolTipFmt", "Remove {0}"), DisplayName))
-					.OnClicked(this, &SBlendProfileMenuEntry::OnRemove)
+					SNew(SComboButton)
+					.ButtonColorAndOpacity(FLinearColor::Transparent)
+					.ButtonStyle(FEditorStyle::Get(), "PropertyEditor.AssetComboStyle")
+					.HasDownArrow(false)
+					.ForegroundColor(FEditorStyle::GetColor("PropertyEditor.AssetName.ColorAndOpacity"))
+					.ContentPadding(2.0f)
+					.ToolTipText(FText::Format(LOCTEXT("ModifyBlendProfileToolTipFmt", "Modify {0}"), DisplayName))
+					.MenuPlacement(MenuPlacement_MenuRight)
+					.OnGetMenuContent(this, &SBlendProfileMenuEntry::GetMenuContent)
+					.ButtonContent()
+					[
+						SNew(SImage)
+						.ColorAndOpacity(FSlateColor::UseForeground())
+						.Image(FAppStyle::Get().GetBrush("Icons.Settings"))
+					]
 				];
 		}
 	}
 
-	FReply OnOpen()
+	void OnProfileModeSelected(EBlendProfileMode InMode)
 	{
-		OnOpenClickedDelegate.ExecuteIfBound();
-		FSlateApplication::Get().DismissAllMenus();
-		return FReply::Handled();
+		OnProfileModeChangedDelegate.ExecuteIfBound(InMode);
+	}
+
+	EBlendProfileMode GetProfileMode()
+	{
+		if (BlendProfile.IsValid())
+		{
+			return BlendProfile->BlendProfileMode;
+		}
+		return EBlendProfileMode::WeightFactor;
+	}
+
+	bool IsProfileModeSelected(EBlendProfileMode InMode)
+	{
+		if (BlendProfile.IsValid())
+		{
+			return BlendProfile->BlendProfileMode == InMode;
+		}
+		return false;
+	}
+
+	TSharedRef<SWidget> GetMenuContent()
+	{
+		FMenuBuilder MenuBuilder(false, nullptr, TSharedPtr<FExtender>(), true);
+
+		EBlendProfileMode BlendProfileMode = GetProfileMode();
+
+		if (BlendProfileMode != EBlendProfileMode::BlendMask)
+		{
+			MenuBuilder.BeginSection(NAME_None, LOCTEXT("BlendProfileMode", "Mode"));
+			{
+				UEnum* ModeEnum = StaticEnum<EBlendProfileMode>();
+				check(ModeEnum);
+
+				// Last enum entry is _MAX
+				int32 NumEnums = ModeEnum->NumEnums() - 1;
+				static const TCHAR* HiddenMeta = TEXT("Hidden");
+
+				for (int32 EnumIndex = 0; EnumIndex < NumEnums; ++EnumIndex)
+				{
+					if (!ModeEnum->HasMetaData(HiddenMeta, EnumIndex))
+					{
+						EBlendProfileMode IndexMode = (EBlendProfileMode)ModeEnum->GetValueByIndex(EnumIndex);
+						MenuBuilder.AddMenuEntry(
+							ModeEnum->GetDisplayNameTextByIndex(EnumIndex),
+							ModeEnum->GetToolTipTextByIndex(EnumIndex),
+							FSlateIcon(),
+							FUIAction(
+								FExecuteAction::CreateSP(this, &SBlendProfileMenuEntry::OnProfileModeSelected, IndexMode),
+								FCanExecuteAction(),
+								FIsActionChecked::CreateSP(this, &SBlendProfileMenuEntry::IsProfileModeSelected, IndexMode)
+							),
+							NAME_None,
+							EUserInterfaceActionType::RadioButton);
+					}
+				}
+			}
+			MenuBuilder.EndSection();
+		}
+
+		MenuBuilder.BeginSection(NAME_None, LOCTEXT("BlendProfileActions", "Actions"));
+		{
+			const FText& ModeName = BlendProfilePickerNames::GetNameForMode(BlendProfileMode);
+			MenuBuilder.AddMenuEntry(
+				FText::Format(LOCTEXT("RemoveBlendProfile", "Remove {0}"), ModeName),
+				FText::Format(LOCTEXT("RemoveBlendProfile_ToolTip", "Remove this {0} from the skeleton"), ModeName),
+				FSlateIcon(),
+				FUIAction(FExecuteAction::CreateLambda([this]()
+					{
+						OnRemoveClickedDelegate.ExecuteIfBound();
+						FSlateApplication::Get().DismissAllMenus();
+					})));
+
+		}
+		MenuBuilder.EndSection();
+
+		return MenuBuilder.MakeWidget();
 	}
 
-	FReply OnRemove()
+	FReply OnOpen()
 	{
-		OnRemoveClickedDelegate.ExecuteIfBound();
+		OnOpenClickedDelegate.ExecuteIfBound();
 		FSlateApplication::Get().DismissAllMenus();
 		return FReply::Handled();
 	}
@@ -121,6 +228,8 @@ private:
 
 	FExecuteAction OnOpenClickedDelegate;
 	FExecuteAction OnRemoveClickedDelegate;
+	FBlendProfileModeChanged OnProfileModeChangedDelegate;
+	TWeakObjectPtr<UBlendProfile> BlendProfile;
 };
 
 //////////////////////////////////////////////////////////////////////////
@@ -132,7 +241,8 @@ void SBlendProfilePicker::Construct(const FArguments& InArgs, TSharedRef<class I
 	bShowClearOption = InArgs._AllowClear;
 	bIsStandalone = InArgs._Standalone;
 	EditableSkeleton = InEditableSkeleton;
-	bAllowRemove = InArgs._AllowRemove;
+	SupportedBlendProfileModes = InArgs._SupportedBlendProfileModes;
+	bAllowModify = InArgs._AllowModify;
 
 	UEditorEngine* Editor = Cast<UEditorEngine>(GEngine);
 	if (Editor != nullptr)
@@ -140,6 +250,23 @@ void SBlendProfilePicker::Construct(const FArguments& InArgs, TSharedRef<class I
 		Editor->RegisterForUndo(this);
 	}
 
+	PropertyHandle = InArgs._PropertyHandle;
+
+	if (PropertyHandle.IsValid())
+	{
+		FSimpleDelegate OnModeChanged = FSimpleDelegate::CreateLambda([this]
+			{
+				UObject * PropertyValue = nullptr;
+				PropertyHandle->GetValue(PropertyValue);
+				UBlendProfile * CurrentProfile = Cast<UBlendProfile>(PropertyValue);
+				// Avoid broadcasting to not double call FBlendProfileCustomization::OnBlendProfileChanged
+				SetSelectedProfile(CurrentProfile, false);
+			});
+
+		// Sometimes, the property value changes externally (not through the blend profile picker i.e. Reset to default). Must notify the blend profile picker
+		PropertyHandle->SetOnPropertyValueChanged(OnModeChanged);
+	}
+
 	if(InArgs._InitialProfile != nullptr && InEditableSkeleton->GetBlendProfiles().Contains(InArgs._InitialProfile))
 	{
 		SelectedProfileName = InArgs._InitialProfile->GetFName();
@@ -170,8 +297,7 @@ void SBlendProfilePicker::Construct(const FArguments& InArgs, TSharedRef<class I
 				.Image(FEditorStyle::GetBrush("SkeletonTree.BlendProfile"))
 			]
 			+SHorizontalBox::Slot()
-			.AutoWidth()
-			.Padding(2, 0, 0, 0)
+			.Padding(2.0f, 0, 8.0f, 0)
 			.VAlign(VAlign_Center)
 			[
 				TextBlock
@@ -213,7 +339,7 @@ FText SBlendProfilePicker::GetSelectedProfileName() const
 	{
 		if (bIsStandalone)
 		{
-			return FText::Format(FText(LOCTEXT("SelectedNameEntryStandalone", "Blend Profile: {0}")), FText::FromName(SelectedProfileName));
+			return FText::Format(FText(LOCTEXT("SelectedNameEntryStandalone", "{0}: {1}")), BlendProfilePickerNames::GetNameForMode(SelectedProfile->BlendProfileMode), FText::FromName(SelectedProfileName));
 		}
 		else
 		{
@@ -222,7 +348,7 @@ FText SBlendProfilePicker::GetSelectedProfileName() const
 	}
 	if (bIsStandalone)
 	{
-		return FText(LOCTEXT("NoSelectionEntryStandalone", "Blend Profile: None"));
+		return FText(LOCTEXT("NoSelectionEntryStandalone", "Blend Profile/Mask: None"));
 	}
 	else
 	{
@@ -243,10 +369,16 @@ TSharedRef<SWidget> SBlendProfilePicker::GetMenuContent()
 			if(bShowNewOption)
 			{
 				MenuBuilder.AddMenuEntry(
-					LOCTEXT("CreateNew", "Create New Blend Profile"),
-					LOCTEXT("CreateNew_ToolTip", "Creates a new blend profile inside the skeleton."),
+					LOCTEXT("CreateNewBlendProfile", "Create New Blend Profile"),
+					LOCTEXT("CreateNewBlendProfile_ToolTip", "Creates a new blend profile inside the skeleton."),
+					FSlateIcon(),
+					FUIAction(FExecuteAction::CreateSP(this, &SBlendProfilePicker::OnCreateNewProfile, EBlendProfileMode::TimeFactor))); //Default blend profiles to time factor
+
+				MenuBuilder.AddMenuEntry(
+					LOCTEXT("CreateNewBlendMask", "Create New Blend Mask"),
+					LOCTEXT("CreateNewBlendMask_ToolTip", "Creates a new blend mask inside the skeleton."),
 					FSlateIcon(),
-					FUIAction(FExecuteAction::CreateSP(this, &SBlendProfilePicker::OnCreateNewProfile)));
+					FUIAction(FExecuteAction::CreateSP(this, &SBlendProfilePicker::OnCreateNewProfile, EBlendProfileMode::BlendMask)));
 			}
 
 			if(bShowClearOption)
@@ -261,25 +393,44 @@ TSharedRef<SWidget> SBlendProfilePicker::GetMenuContent()
 		MenuBuilder.EndSection();
 	}
 
-	MenuBuilder.BeginSection(NAME_None, LOCTEXT("Profiles", "Available Blend Profiles"));
+	const static FText AvailableModesHeadings[] = { LOCTEXT("BlendProfiles", "Available Blend Profiles"), LOCTEXT("BlendMasks", "Available Blend Masks") };
+	const static EBlendProfilePickerMode AvailableModesFilters[] = { EBlendProfilePickerMode::BlendProfile, EBlendProfilePickerMode::BlendMask };
+	static_assert(UE_ARRAY_COUNT(AvailableModesHeadings) == UE_ARRAY_COUNT(AvailableModesFilters), "Arrays must be of same size");
+
+	for (int32 ModeIndex = 0; ModeIndex < UE_ARRAY_COUNT(AvailableModesHeadings); ++ModeIndex)
 	{
-		if (EditableSkeleton.IsValid())
+		// Skip section if we don't support it
+		if (EnumHasAnyFlags(SupportedBlendProfileModes, AvailableModesFilters[ModeIndex]))
 		{
-			for (UBlendProfile* Profile : EditableSkeleton->GetBlendProfiles())
+			// Note: Section won't get populated if there are no available items for this mode type
+			MenuBuilder.BeginSection(NAME_None, AvailableModesHeadings[ModeIndex]);
 			{
-				MenuBuilder.AddWidget(
-					SNew(SBlendProfileMenuEntry)
-						.LabelOverride(FText::FromString(Profile->GetName()))
-						.OnOpenClickedDelegate(FExecuteAction::CreateSP(this, &SBlendProfilePicker::OnProfileSelected, Profile->GetFName()))
-						.OnRemoveClickedDelegate(FExecuteAction::CreateSP(this, &SBlendProfilePicker::OnProfileRemoved, Profile->GetFName()))
-						.AllowRemove(bAllowRemove),
-					FText(),
-					true
-				);
+				if (EditableSkeleton.IsValid())
+				{
+					for (UBlendProfile* Profile : EditableSkeleton->GetBlendProfiles())
+					{
+						// Only display blend profiles w/ a matching mode
+						if ((AvailableModesFilters[ModeIndex] == EBlendProfilePickerMode::BlendMask) ? Profile->IsBlendMask() : !Profile->IsBlendMask())
+						{
+							MenuBuilder.AddWidget(
+								SNew(SBlendProfileMenuEntry)
+								.LabelOverride(FText::FromString(Profile->GetName()))
+								.OnOpenClickedDelegate(FExecuteAction::CreateSP(this, &SBlendProfilePicker::OnProfileSelected, Profile->GetFName()))
+								.OnRemoveClickedDelegate(FExecuteAction::CreateSP(this, &SBlendProfilePicker::OnProfileRemoved, Profile->GetFName()))
+								.OnProfileModeChangedDelegate(SBlendProfileMenuEntry::FBlendProfileModeChanged::CreateSP(this, &SBlendProfilePicker::OnProfileModeChanged, Profile->GetFName()))
+								.BlendProfile(MakeWeakObjectPtr(Profile))
+								.AllowModify(bAllowModify),
+								FText(),
+								true
+							);
+						}
+					}
+				}
 			}
 		}
+
+		MenuBuilder.EndSection();
 	}
-	MenuBuilder.EndSection();
 
 	return MenuBuilder.MakeWidget();
 }
@@ -303,11 +454,16 @@ void SBlendProfilePicker::OnProfileRemoved(FName InBlendProfileName)
 	BlendProfileSelectedDelegate.ExecuteIfBound(nullptr);
 }
 
-void SBlendProfilePicker::OnCreateNewProfile()
+void SBlendProfilePicker::OnProfileModeChanged(EBlendProfileMode ProfileMode, FName InBlendProfileName)
+{
+	EditableSkeleton->SetBlendProfileMode(InBlendProfileName, ProfileMode);
+}
+
+void SBlendProfilePicker::OnCreateNewProfile(EBlendProfileMode InMode)
 {
 	TSharedRef<STextEntryPopup> TextEntry = SNew(STextEntryPopup)
 		.Label(LOCTEXT("NewProfileName", "Profile Name"))
-		.OnTextCommitted(this, &SBlendProfilePicker::OnCreateNewProfileComitted);
+		.OnTextCommitted(this, &SBlendProfilePicker::OnCreateNewProfileComitted, InMode);
 
 	FSlateApplication::Get().PushMenu(
 		AsShared(),
@@ -317,7 +473,7 @@ void SBlendProfilePicker::OnCreateNewProfile()
 		FPopupTransitionEffect(FPopupTransitionEffect::TypeInPopup));
 }
 
-void SBlendProfilePicker::OnCreateNewProfileComitted(const FText& NewName, ETextCommit::Type CommitType)
+void SBlendProfilePicker::OnCreateNewProfileComitted(const FText& NewName, ETextCommit::Type CommitType, EBlendProfileMode InMode)
 {
 	FSlateApplication::Get().DismissAllMenus();
 
@@ -334,6 +490,8 @@ void SBlendProfilePicker::OnCreateNewProfileComitted(const FText& NewName, EText
 		}
 		else if(UBlendProfile* NewProfile = EditableSkeleton->CreateNewBlendProfile(NameToUse))
 		{
+			// Set our initial blend profile mode. Blend masks can't change this.
+			NewProfile->BlendProfileMode = InMode;
 			OnProfileSelected(NewProfile->GetFName());
 		}
 	}
diff --git a/Engine/Source/Editor/SkeletonEditor/Private/SBlendProfilePicker.h b/Engine/Source/Editor/SkeletonEditor/Private/SBlendProfilePicker.h
--- a/Engine/Source/Editor/SkeletonEditor/Private/SBlendProfilePicker.h
+++ b/Engine/Source/Editor/SkeletonEditor/Private/SBlendProfilePicker.h
@@ -11,6 +11,7 @@
 #include "BlendProfilePicker.h"
 
 class UBlendProfile;
+enum class EBlendProfileMode : uint8;
 
 // Picker for UBlendProfile instances inside a USkeleton
 class SBlendProfilePicker : public SCompoundWidget, public FEditorUndoClient
@@ -19,22 +20,28 @@ public:
 	SLATE_BEGIN_ARGS(SBlendProfilePicker)
 		: _InitialProfile(nullptr)
 		, _OnBlendProfileSelected()
+		, _SupportedBlendProfileModes(EBlendProfilePickerMode::AllModes)
 		, _AllowNew(true)
 		, _AllowClear(true)
-		, _AllowRemove(true)
+		, _AllowModify(true)
+		, _Standalone(false)
 	{}
 		// Initial blend profile selected
 		SLATE_ARGUMENT(UBlendProfile*, InitialProfile)
 		// Delegate to call when the picker selection is changed
 		SLATE_EVENT(FOnBlendProfileSelected, OnBlendProfileSelected)
+		// Only display Blend Profiles w/ specified Blend Profile modes (EBlendProfilePickerMode values are flags.)
+		SLATE_ARGUMENT(EBlendProfilePickerMode, SupportedBlendProfileModes)
 		// Allow the option to create new profiles in the picker
 		SLATE_ARGUMENT(bool, AllowNew)
 		// Allow the option to clear the profile selection
 		SLATE_ARGUMENT(bool, AllowClear)
-		// Allow the option to delete from the skeleton
-		SLATE_ARGUMENT(bool, AllowRemove)
+		// Allow the option to modify (remove/edit settings) the profile from the skeleton
+		SLATE_ARGUMENT(bool, AllowModify)
 		// Is this a standalone blend profile picker?
 		SLATE_ARGUMENT(bool, Standalone)
+		// Optional property handle using this widget
+		SLATE_ARGUMENT(TSharedPtr<class IPropertyHandle>, PropertyHandle)
 	SLATE_END_ARGS()
 
 	~SBlendProfilePicker();
@@ -53,8 +60,8 @@ public:
 	/** Get the currently selected blend profile name */
 	FName GetSelectedBlendProfileName() const;
 
-	/** Create a New Blend profile withe the provided profile name */
-	void OnCreateNewProfileComitted(const FText& NewName, ETextCommit::Type CommitType);
+	/** Create a New Blend profile withe the provided profile name and mode*/
+	void OnCreateNewProfileComitted(const FText& NewName, ETextCommit::Type CommitType, EBlendProfileMode InMode);
 	
 	/* Deselect the current blend profile */
 	void OnClearSelection();
@@ -65,10 +72,11 @@ private:
 	virtual void PostUndo(bool bSuccess) override;
 	virtual void PostRedo(bool bSuccess) override;
 
-	void OnCreateNewProfile();
+	void OnCreateNewProfile(EBlendProfileMode InMode);
 
 	void OnProfileSelected(FName InBlendProfileName);
 	void OnProfileRemoved(FName InBlendProfileName);
+	void OnProfileModeChanged(EBlendProfileMode ProfileMode, FName InBlendProfileName);
 
 	FText GetSelectedProfileName() const;
 
@@ -77,12 +85,19 @@ private:
 	bool bShowNewOption;
 	bool bShowClearOption;
 	bool bIsStandalone;
-	bool bAllowRemove;
+	bool bAllowModify;
+	bool bAllowBlendMask;
+	bool bAllowOnlyBlendMask;
+
+	/** Only display Blend Profiles w/ specified Blend Profile modes (EBlendProfilePickerMode values are flags.) */
+	EBlendProfilePickerMode SupportedBlendProfileModes;
 
 	FName SelectedProfileName;
 
 	TSharedPtr<class IEditableSkeleton> EditableSkeleton;
 
+	TSharedPtr<class IPropertyHandle> PropertyHandle;
+
 	FOnBlendProfileSelected BlendProfileSelectedDelegate;
 
 };
diff --git a/Engine/Source/Editor/SkeletonEditor/Private/SSkeletonTree.cpp b/Engine/Source/Editor/SkeletonEditor/Private/SSkeletonTree.cpp
--- a/Engine/Source/Editor/SkeletonEditor/Private/SSkeletonTree.cpp
+++ b/Engine/Source/Editor/SkeletonEditor/Private/SSkeletonTree.cpp
@@ -606,7 +606,7 @@ void SSkeletonTree::CreateTreeColumns()
 				FName CurrentProfile = BlendProfilePicker->GetSelectedBlendProfileName();
 				return CurrentProfile != NAME_None ? FText::FromName(CurrentProfile) : LOCTEXT("NoBlendProfile", "No Blend");
 			})
-			.OnTextCommitted(BlendProfilePicker.ToSharedRef(), &SBlendProfilePicker::OnCreateNewProfileComitted)
+			.OnTextCommitted(BlendProfilePicker.ToSharedRef(), &SBlendProfilePicker::OnCreateNewProfileComitted, EBlendProfileMode::TimeFactor)
 			.IsReadOnly(true)
 		]
 	];
diff --git a/Engine/Source/Editor/SkeletonEditor/Public/BlendProfilePicker.h b/Engine/Source/Editor/SkeletonEditor/Public/BlendProfilePicker.h
--- a/Engine/Source/Editor/SkeletonEditor/Public/BlendProfilePicker.h
+++ b/Engine/Source/Editor/SkeletonEditor/Public/BlendProfilePicker.h
@@ -8,6 +8,14 @@ class UBlendProfile;
 
 DECLARE_DELEGATE_OneParam(FOnBlendProfileSelected, UBlendProfile*);
 
+enum class EBlendProfilePickerMode : uint8
+{
+	BlendProfile = (1U << 0),
+	BlendMask = (1U << 1),
+	AllModes = BlendProfile | BlendMask,
+};
+ENUM_CLASS_FLAGS(EBlendProfilePickerMode);
+
 /** Argument used to set up the blend profile picker */
 struct FBlendProfilePickerArgs
 {
@@ -24,13 +32,21 @@ struct FBlendProfilePickerArgs
 	bool bAllowClear;
 
 	/** Allow the option to delete blend profiles from the skeleton */
-	bool bAllowRemove;
+	bool bAllowModify;
+
+	/** Only display Blend Profiles w/ specified Blend Profile modes (EBlendProfilePickerMode values are flags.) */
+	EBlendProfilePickerMode SupportedBlendProfileModes;
+
+	/** Optional property handle, when the picker is tied to a property */
+	TSharedPtr<class IPropertyHandle> PropertyHandle;
 
 	FBlendProfilePickerArgs()
 		: InitialProfile(nullptr)
 		, OnBlendProfileSelected()
 		, bAllowNew(false)
 		, bAllowClear(false)
-		, bAllowRemove(false)
+		, bAllowModify(false)
+		, SupportedBlendProfileModes(EBlendProfilePickerMode::AllModes)
+		, PropertyHandle(nullptr)
 	{}
 };
diff --git a/Engine/Source/Editor/SkeletonEditor/Public/IEditableSkeleton.h b/Engine/Source/Editor/SkeletonEditor/Public/IEditableSkeleton.h
--- a/Engine/Source/Editor/SkeletonEditor/Public/IEditableSkeleton.h
+++ b/Engine/Source/Editor/SkeletonEditor/Public/IEditableSkeleton.h
@@ -6,6 +6,7 @@
 #include "BoneContainer.h"
 #include "Animation/Skeleton.h"
 
+enum class EBlendProfileMode : uint8;
 class UBlendProfile;
 class USkeletalMesh;
 
@@ -41,6 +42,9 @@ public:
 	/** Set the blend profile scale for the specified bone */
 	virtual void SetBlendProfileScale(const FName& InBlendProfileName, const FName& InBoneName, float InNewScale, bool bInRecurse) = 0;
 
+	/** Change an existing blend profile's mode (See EBlendProfileMode for details) */
+	virtual void SetBlendProfileMode(FName InBlendProfileName, EBlendProfileMode ProfileMode) = 0;
+
 	/** Creates a new socket on the skeleton */
 	virtual USkeletalMeshSocket* AddSocket(const FName& InBoneName) = 0;
 
diff --git a/Engine/Source/Runtime/AnimGraphRuntime/Private/AnimNodes/AnimNode_LayeredBoneBlend.cpp b/Engine/Source/Runtime/AnimGraphRuntime/Private/AnimNodes/AnimNode_LayeredBoneBlend.cpp
--- a/Engine/Source/Runtime/AnimGraphRuntime/Private/AnimNodes/AnimNode_LayeredBoneBlend.cpp
+++ b/Engine/Source/Runtime/AnimGraphRuntime/Private/AnimNodes/AnimNode_LayeredBoneBlend.cpp
@@ -41,7 +41,15 @@ void FAnimNode_LayeredBoneBlend::RebuildCacheData(const USkeleton* InSkeleton)
 {
 	if (InSkeleton)
 	{
-		FAnimationRuntime::CreateMaskWeights(PerBoneBlendWeights, LayerSetup, InSkeleton);
+		if (BlendMode == ELayeredBoneBlendMode::BranchFilter)
+		{
+			FAnimationRuntime::CreateMaskWeights(PerBoneBlendWeights, LayerSetup, InSkeleton);
+		}
+		else
+		{
+			FAnimationRuntime::CreateMaskWeights(PerBoneBlendWeights, BlendMasks, InSkeleton);
+		}
+
 		SkeletonGuid = InSkeleton->GetGuid();
 		VirtualBoneGuid = InSkeleton->GetVirtualBoneGuid();
 	}
@@ -322,7 +330,7 @@ void FAnimNode_LayeredBoneBlend::ValidateData()
 	// so here we add code to fix this up manually in editor, so that they can continue working on it. 
 	int32 PoseNum = BlendPoses.Num();
 	int32 WeightNum = BlendWeights.Num();
-	int32 LayerNum = LayerSetup.Num();
+	int32 LayerNum = (BlendMode == ELayeredBoneBlendMode::BranchFilter) ? LayerSetup.Num() : BlendMasks.Num();
 
 	int32 Max = FMath::Max3(PoseNum, WeightNum, LayerNum);
 	int32 Min = FMath::Min3(PoseNum, WeightNum, LayerNum);
@@ -341,9 +349,30 @@ void FAnimNode_LayeredBoneBlend::ValidateData()
 			BlendWeights.Add(1.f);
 		}
 
-		for(int32 Index=LayerNum; Index<Max; ++Index)
+		if (BlendMode == ELayeredBoneBlendMode::BranchFilter)
+		{
+			for (int32 Index = LayerNum; Index < Max; ++Index)
+			{
+				LayerSetup.Add(FInputBlendPose());
+			}
+		}
+		else
+		{
+			// Get rid of our branch filters if we are not using them to avoid saving into the asset
+			LayerSetup.Reset();
+		}
+
+		if (BlendMode == ELayeredBoneBlendMode::BlendMask)
+		{
+			for (int32 Index = LayerNum; Index < Max; ++Index)
+			{
+				BlendMasks.Add(nullptr);
+			}
+		}
+		else
 		{
-			LayerSetup.Add(FInputBlendPose());
+			// Get rid of blend masks array if we are not using them to avoid saving into the asset
+			BlendMasks.Reset();
 		}
 	}
 }
diff --git a/Engine/Source/Runtime/AnimGraphRuntime/Public/AnimNodes/AnimNode_BlendListBase.h b/Engine/Source/Runtime/AnimGraphRuntime/Public/AnimNodes/AnimNode_BlendListBase.h
--- a/Engine/Source/Runtime/AnimGraphRuntime/Public/AnimNodes/AnimNode_BlendListBase.h
+++ b/Engine/Source/Runtime/AnimGraphRuntime/Public/AnimNodes/AnimNode_BlendListBase.h
@@ -48,7 +48,7 @@ public:
 	UPROPERTY(EditAnywhere, Category=BlendType)
 	UCurveFloat* CustomBlendCurve;
 
-	UPROPERTY(EditAnywhere, Category=BlendType)
+	UPROPERTY(EditAnywhere, Category=BlendType, meta=(UseAsBlendProfile=true))
 	UBlendProfile* BlendProfile;
 
 	TArray<struct FAlphaBlend> Blends;
diff --git a/Engine/Source/Runtime/AnimGraphRuntime/Public/AnimNodes/AnimNode_LayeredBoneBlend.h b/Engine/Source/Runtime/AnimGraphRuntime/Public/AnimNodes/AnimNode_LayeredBoneBlend.h
--- a/Engine/Source/Runtime/AnimGraphRuntime/Public/AnimNodes/AnimNode_LayeredBoneBlend.h
+++ b/Engine/Source/Runtime/AnimGraphRuntime/Public/AnimNodes/AnimNode_LayeredBoneBlend.h
@@ -9,6 +9,13 @@
 #include "Animation/AnimData/BoneMaskFilter.h"
 #include "AnimNode_LayeredBoneBlend.generated.h"
 
+UENUM()
+enum class ELayeredBoneBlendMode : uint8
+{
+	BranchFilter,
+	BlendMask,
+};
+
 // Layered blend (per bone); has dynamic number of blendposes that can blend per different bone sets
 USTRUCT(BlueprintInternalUseOnly)
 struct ANIMGRAPHRUNTIME_API FAnimNode_LayeredBoneBlend : public FAnimNode_Base
@@ -24,9 +31,21 @@ public:
 	UPROPERTY(EditAnywhere, BlueprintReadWrite, editfixedsize, Category=Links)
 	TArray<FPoseLink> BlendPoses;
 
+	/** Whether to use branch filters or a blend mask to specify an input pose per-bone influence */
+	UPROPERTY(EditAnywhere, Category = Config)
+	ELayeredBoneBlendMode BlendMode;
+
+	/** 
+	 * The blend masks to use for our layer inputs. Allows the use of per-bone alphas.
+	 * Blend masks are used when BlendMode is BlendMask.
+	 */
+	UPROPERTY(EditAnywhere, editfixedsize, Category=Config, meta=(UseAsBlendMask=true))
+	TArray<UBlendProfile*> BlendMasks;
+
 	/** 
 	 * Configuration for the parts of the skeleton to blend for each layer. Allows
 	 * certain parts of the tree to be blended out or omitted from the pose.
+	 * LayerSetup is used when BlendMode is BranchFilter.
 	 */
 	UPROPERTY(EditAnywhere, editfixedsize, Category=Config)
 	TArray<FInputBlendPose> LayerSetup;
diff --git a/Engine/Source/Runtime/Engine/Classes/Animation/BlendProfile.h b/Engine/Source/Runtime/Engine/Classes/Animation/BlendProfile.h
--- a/Engine/Source/Runtime/Engine/Classes/Animation/BlendProfile.h
+++ b/Engine/Source/Runtime/Engine/Classes/Animation/BlendProfile.h
@@ -24,9 +24,11 @@ enum class EBlendProfileMode : uint8
 	// Values should typically be equal or greater than 1.0.
 	// If you want certain bones to instantly transition into the target state
 	// the Time Factor based method might be a better choice.
-	WeightFactor
-};
+	WeightFactor,
 
+	// Used for blend masks. Per bone alpha
+	BlendMask UMETA(Hidden),
+};
 
 /** A single entry for a blend scale within a profile, mapping a bone to a blendscale */
 USTRUCT()
@@ -108,6 +110,11 @@ public:
 	virtual void PostLoad() override;
 	// End UObject
 
+	// Default value of entries. Default values are not saved
+	virtual float GetDefaultBlendScale() const { return IsBlendMask() ? 0.0f : 1.0f; }
+
+	bool IsBlendMask() const { return BlendProfileMode == EBlendProfileMode::BlendMask;  }
+
 private:
 
 	/** Sets the skeleton this blend profile is used with */
@@ -124,4 +131,8 @@ public:
 	// List of blend scale entries
 	UPROPERTY()
 	TArray<FBlendProfileBoneEntry> ProfileEntries;
+
+	// Blend Profile Mode. Read EBlendProfileMode for more details
+	UPROPERTY()
+	EBlendProfileMode BlendProfileMode;
 };
diff --git a/Engine/Source/Runtime/Engine/Private/Animation/AnimationRuntime.cpp b/Engine/Source/Runtime/Engine/Private/Animation/AnimationRuntime.cpp
--- a/Engine/Source/Runtime/Engine/Private/Animation/AnimationRuntime.cpp
+++ b/Engine/Source/Runtime/Engine/Private/Animation/AnimationRuntime.cpp
@@ -13,6 +13,7 @@
 #include "Animation/CustomAttributesRuntime.h"
 #include "GenericPlatform/GenericPlatformCompilerPreSetup.h"
 #include "Animation/AnimationPoseData.h"
+#include "Animation/BlendProfile.h"
 #if INTEL_ISPC
 #include "AnimationRuntime.ispc.generated.h"
 #endif
@@ -1819,6 +1820,39 @@ void FAnimationRuntime::CreateMaskWeights(TArray<FPerBoneBlendWeight>& BoneBlend
 	}
 }
 
+void FAnimationRuntime::CreateMaskWeights(TArray<FPerBoneBlendWeight>& BoneBlendWeights, const TArray<class UBlendProfile*>& BlendMasks, const USkeleton* Skeleton)
+{
+	if (Skeleton)
+	{
+		const FReferenceSkeleton& RefSkeleton = Skeleton->GetReferenceSkeleton();
+
+		const int32 NumBones = RefSkeleton.GetNum();
+		BoneBlendWeights.Reset(NumBones);
+		// We only store non-zero weights in blend masks. Initialize all to zero.
+		BoneBlendWeights.AddZeroed(NumBones);
+
+		for (int32 MaskIndex = 0; MaskIndex < BlendMasks.Num(); ++MaskIndex)
+		{
+			const UBlendProfile* BlendMask = BlendMasks[MaskIndex];
+
+			if (!BlendMask || BlendMask->BlendProfileMode != EBlendProfileMode::BlendMask)
+			{
+				ensureMsgf(false, TEXT("FAnimationRuntime::CreateMaskWeights BlendMask null or BlendProfile mode is not blend mask.  BlendProfile=%s"), *GetNameSafe(BlendMask));
+				continue;
+			}
+
+			for (int32 EntryIndex = 0; EntryIndex < BlendMask->GetNumBlendEntries(); EntryIndex++)
+			{
+				int32 BoneIndex = BlendMask->ProfileEntries[EntryIndex].BoneReference.BoneIndex;
+				if (BoneBlendWeights.IsValidIndex(BoneIndex))
+				{
+					BoneBlendWeights[BoneIndex].BlendWeight = BlendMask->ProfileEntries[EntryIndex].BlendScale;
+				}
+			}
+		}
+	}
+}
+
 void FAnimationRuntime::ConvertCSTransformToBoneSpace(USkeletalMeshComponent* SkelComp, FCSPose<FCompactPose>& MeshBases, FTransform& InOutCSBoneTM, FCompactPoseBoneIndex BoneIndex, EBoneControlSpace Space)
 {
 	ConvertCSTransformToBoneSpace(SkelComp->GetComponentTransform(), MeshBases, InOutCSBoneTM, BoneIndex, Space);
diff --git a/Engine/Source/Runtime/Engine/Private/Animation/BlendProfile.cpp b/Engine/Source/Runtime/Engine/Private/Animation/BlendProfile.cpp
--- a/Engine/Source/Runtime/Engine/Private/Animation/BlendProfile.cpp
+++ b/Engine/Source/Runtime/Engine/Private/Animation/BlendProfile.cpp
@@ -4,6 +4,7 @@
 
 UBlendProfile::UBlendProfile()
 	: OwningSkeleton(nullptr)
+	, BlendProfileMode(EBlendProfileMode::WeightFactor)
 {
 	// Set up our owning skeleton and initialise bone references
 	if(USkeleton* OuterAsSkeleton = Cast<USkeleton>(GetOuter()))
@@ -50,7 +51,7 @@ float UBlendProfile::GetBoneBlendScale(int32 InBoneIdx) const
 		return FoundEntry->BlendScale;
 	}
 
-	return 1.0f;
+	return GetDefaultBlendScale();
 }
 
 float UBlendProfile::GetBoneBlendScale(const FName& InBoneName) const
@@ -65,7 +66,7 @@ float UBlendProfile::GetBoneBlendScale(const FName& InBoneName) const
 		return FoundEntry->BlendScale;
 	}
 
-	return 1.0f;
+	return GetDefaultBlendScale();
 }
 
 void UBlendProfile::SetSkeleton(USkeleton* InSkeleton)
@@ -129,7 +130,7 @@ float UBlendProfile::GetEntryBlendScale(const int32 InEntryIdx) const
 		return ProfileEntries[InEntryIdx].BlendScale;
 	}
 	// No overriden blend scale, return no scale
-	return 1.0f;
+	return GetDefaultBlendScale();
 }
 
 int32 UBlendProfile::GetPerBoneInterpolationIndex(int32 BoneIndex, const FBoneContainer& RequiredBones) const
@@ -160,8 +161,8 @@ void UBlendProfile::SetSingleBoneBlendScale(int32 InBoneIdx, float InScale, bool
 		Modify();
 		Entry->BlendScale = InScale;
 
-		// Remove any entry that gets set back to 1.0f - so we only store entries that actually contain a scale
-		if(Entry->BlendScale == 1.0f)
+		// Remove any entry that gets set back to DefautBlendScale - so we only store entries that actually contain a scale
+		if(Entry->BlendScale == GetDefaultBlendScale())
 		{
 			ProfileEntries.RemoveAll([InBoneIdx](const FBlendProfileBoneEntry& Current)
 			{
diff --git a/Engine/Source/Runtime/Engine/Public/AnimationRuntime.h b/Engine/Source/Runtime/Engine/Public/AnimationRuntime.h
--- a/Engine/Source/Runtime/Engine/Public/AnimationRuntime.h
+++ b/Engine/Source/Runtime/Engine/Public/AnimationRuntime.h
@@ -332,8 +332,8 @@ public:
 	static void UpdateDesiredBoneWeight(const TArray<FPerBoneBlendWeight>& SrcBoneBlendWeights, TArray<FPerBoneBlendWeight>& TargetBoneBlendWeights, const TArray<float>& BlendWeights);
 
 	/**
-	 *	Create Mast Weight for skeleton joints, not per mesh or per required bones
-	 *  You'll have to filter properly with correct mesh joint or required boens
+	 *	Create Mask Weight for skeleton joints, not per mesh or per required bones
+	 *  You'll have to filter properly with correct mesh joint or required bones
 	 *  The depth should not change based on LOD or mesh or skeleton
 	 *	They still should contain same depth
 	 */
@@ -342,6 +342,15 @@ public:
 			const TArray<FInputBlendPose>& BlendFilters, 
 			const USkeleton* Skeleton);
 
+	/**
+	 *	Create Mask Weight for skeleton joints, not per mesh or per required bones
+	 *  Individual alphas are read from a BlendProfile using a BlendMask mode
+	 */
+	static void CreateMaskWeights(
+		TArray<FPerBoneBlendWeight>& BoneBlendWeights,
+		const TArray<class UBlendProfile*>& BlendMasks,
+		const USkeleton* Skeleton);
+
 	static void CombineWithAdditiveAnimations(
 		int32 NumAdditivePoses,
 		const FTransformArrayA2** SourceAdditivePoses,
