Missing file from 14828423 for deprecation fix-up

diff --git a/Engine/Source/Editor/SkeletonEditor/Private/EditableSkeleton.cpp b/Engine/Source/Editor/SkeletonEditor/Private/EditableSkeleton.cpp
--- a/Engine/Source/Editor/SkeletonEditor/Private/EditableSkeleton.cpp
+++ b/Engine/Source/Editor/SkeletonEditor/Private/EditableSkeleton.cpp
@@ -321,8 +321,9 @@ void FEditableSkeleton::RenameSmartname(const FName InContainerName, SmartName::
 						if (UAnimSequence* Seq = Cast<UAnimSequence>(SequenceBase))
 						{
 							SequencesToRecompress.Add(Seq);
-
+							PRAGMA_DISABLE_DEPRECATION_WARNINGS
 							Seq->ClearCompressedCurveData();
+							PRAGMA_ENABLE_DEPRECATION_WARNINGS
 						}
 					}
 				}
