Added first pass on blend profile modes for skeleton editor

diff --git a/Engine/Source/Editor/SkeletonEditor/Private/SSkeletonTree.cpp b/Engine/Source/Editor/SkeletonEditor/Private/SSkeletonTree.cpp
--- a/Engine/Source/Editor/SkeletonEditor/Private/SSkeletonTree.cpp
+++ b/Engine/Source/Editor/SkeletonEditor/Private/SSkeletonTree.cpp
@@ -480,8 +480,16 @@ void SSkeletonTree::BindCommands()
 		FExecuteAction::CreateSP(this, &SSkeletonTree::HandleFocusCamera));
 
 	CommandList.MapAction(
-		MenuActions.CreateBlendProfile,
-		FExecuteAction::CreateSP( this, &SSkeletonTree::OnCreateBlendProfile));
+		MenuActions.CreateTimeBlendProfile,
+		FExecuteAction::CreateSP( this, &SSkeletonTree::OnCreateBlendProfile, EBlendProfileMode::TimeFactor));
+
+	CommandList.MapAction(
+		MenuActions.CreateWeightBlendProfile,
+		FExecuteAction::CreateSP(this, &SSkeletonTree::OnCreateBlendProfile, EBlendProfileMode::WeightFactor));
+
+	CommandList.MapAction(
+		MenuActions.CreateBlendMask,
+		FExecuteAction::CreateSP(this, &SSkeletonTree::OnCreateBlendProfile, EBlendProfileMode::BlendMask));
 
 	CommandList.MapAction(
 		MenuActions.DeleteCurrentBlendProfile,
@@ -606,7 +614,10 @@ void SSkeletonTree::CreateTreeColumns()
 				FName CurrentProfile = BlendProfilePicker->GetSelectedBlendProfileName();
 				return CurrentProfile != NAME_None ? FText::FromName(CurrentProfile) : LOCTEXT("NoBlendProfile", "No Blend");
 			})
-			.OnTextCommitted(BlendProfilePicker.ToSharedRef(), &SBlendProfilePicker::OnCreateNewProfileComitted, EBlendProfileMode::TimeFactor)
+			.OnTextCommitted_Lambda([this](const FText& InText, ETextCommit::Type InCommitType)
+			{
+				BlendProfilePicker->OnCreateNewProfileComitted(InText, InCommitType, NewBlendProfileMode);
+			})
 			.IsReadOnly(true)
 		]
 	];
@@ -1620,7 +1631,9 @@ void SSkeletonTree::GetBlendProfileMenu(FMenuBuilder& MenuBuilder)
 		FUIAction(FExecuteAction::CreateSP(BlendProfilePicker.ToSharedRef(), &SBlendProfilePicker::OnClearSelection)));
 
 		MenuBuilder.AddSeparator();
-		MenuBuilder.AddMenuEntry(Actions.CreateBlendProfile);
+		MenuBuilder.AddMenuEntry(Actions.CreateTimeBlendProfile);
+		MenuBuilder.AddMenuEntry(Actions.CreateWeightBlendProfile);
+		MenuBuilder.AddMenuEntry(Actions.CreateBlendMask);
 		if (BlendProfilePicker->GetSelectedBlendProfileName() != NAME_None)
 		{
 			MenuBuilder.AddSeparator();
@@ -1632,12 +1645,15 @@ void SSkeletonTree::GetBlendProfileMenu(FMenuBuilder& MenuBuilder)
 	MenuBuilder.EndSection();
 } 
 
-void SSkeletonTree::OnCreateBlendProfile()
+void SSkeletonTree::OnCreateBlendProfile(const EBlendProfileMode InMode)
 {
 	// Ensure the Blend Profile Column is Visible
 	BlendProfilePicker->OnClearSelection();
 	SkeletonTreeView->GetHeaderRow()->SetShowGeneratedColumn(ISkeletonTree::Columns::BlendProfile);
 
+	// Set our NewBlendProfileMode for our BlendProfileHeader to use when the text is commited.
+	NewBlendProfileMode = InMode;
+
 	// Activate the Header Entry Box
 	BlendProfileHeader->SetReadOnly(false);
 	BlendProfileHeader->EnterEditingMode();
@@ -1764,7 +1780,9 @@ TSharedRef< SWidget > SSkeletonTree::CreateNewMenu()
 
 	MenuBuilder.BeginSection("Blend");
 	{
-		MenuBuilder.AddMenuEntry(Actions.CreateBlendProfile);
+		MenuBuilder.AddMenuEntry(Actions.CreateTimeBlendProfile);
+		MenuBuilder.AddMenuEntry(Actions.CreateWeightBlendProfile);
+		MenuBuilder.AddMenuEntry(Actions.CreateBlendMask);
 	}
 	MenuBuilder.EndSection();
 
diff --git a/Engine/Source/Editor/SkeletonEditor/Private/SSkeletonTree.h b/Engine/Source/Editor/SkeletonEditor/Private/SSkeletonTree.h
--- a/Engine/Source/Editor/SkeletonEditor/Private/SSkeletonTree.h
+++ b/Engine/Source/Editor/SkeletonEditor/Private/SSkeletonTree.h
@@ -41,6 +41,7 @@ struct FNotificationInfo;
 class IPinnedCommandList;
 class FPackageReloadedEvent;
 enum class EPackageReloadPhase : uint8;
+enum class EBlendProfileMode : uint8;
 
 //////////////////////////////////////////////////////////////////////////
 // SSkeletonTree
@@ -324,7 +325,7 @@ private:
 	void HandlePackageReloaded(const EPackageReloadPhase InPackageReloadPhase, FPackageReloadedEvent* InPackageReloadedEvent);
 
 	/** Creates a new Blend Profile */
-	void OnCreateBlendProfile();
+	void OnCreateBlendProfile(const EBlendProfileMode InMode);
 
 	/** Removes the active Blend Profile */
 	void OnDeleteCurrentBlendProfile();
@@ -374,6 +375,9 @@ private:
 	/** Commands that are bound to delegates*/
 	TSharedPtr<FUICommandList_Pinnable> UICommandList;
 
+	/** Current type of blend profile no create. We shouldn't need to hold state for this, but blend profile creation is tied to its header text committed*/
+	EBlendProfileMode NewBlendProfileMode;
+
 	/** Current type of bones to show */
 	EBoneFilter BoneFilter;
 
diff --git a/Engine/Source/Editor/SkeletonEditor/Private/SkeletonTreeCommands.cpp b/Engine/Source/Editor/SkeletonEditor/Private/SkeletonTreeCommands.cpp
--- a/Engine/Source/Editor/SkeletonEditor/Private/SkeletonTreeCommands.cpp
+++ b/Engine/Source/Editor/SkeletonEditor/Private/SkeletonTreeCommands.cpp
@@ -36,7 +36,9 @@ void FSkeletonTreeCommands::RegisterCommands()
 
 	UI_COMMAND( FocusCamera, "Focus Camera", "Focus the camera on the current selection", EUserInterfaceActionType::Button, FInputChord(EKeys::F) );
 
-	UI_COMMAND( CreateBlendProfile, "Add Blend Profile", "Add a new blend profile for this skeleton", EUserInterfaceActionType::Button, FInputChord() );
+	UI_COMMAND( CreateTimeBlendProfile, "Add Time Blend Profile", "Add a new time blend profile for this skeleton", EUserInterfaceActionType::Button, FInputChord() );
+	UI_COMMAND( CreateWeightBlendProfile, "Add Weight Blend Profile", "Add a new weight blend profile for this skeleton", EUserInterfaceActionType::Button, FInputChord() );
+	UI_COMMAND( CreateBlendMask, "Add Blend Mask", "Add a new blend mask for this skeleton", EUserInterfaceActionType::Button, FInputChord() );
 	UI_COMMAND( DeleteCurrentBlendProfile, "Delete the Current Blend Profile", "Delete the active blend profile for this skeleton", EUserInterfaceActionType::Button, FInputChord() );
 
 }
diff --git a/Engine/Source/Editor/SkeletonEditor/Private/SkeletonTreeCommands.h b/Engine/Source/Editor/SkeletonEditor/Private/SkeletonTreeCommands.h
--- a/Engine/Source/Editor/SkeletonEditor/Private/SkeletonTreeCommands.h
+++ b/Engine/Source/Editor/SkeletonEditor/Private/SkeletonTreeCommands.h
@@ -99,8 +99,14 @@ public:
 	/** Focus the camera on the current selection */
 	TSharedPtr< FUICommandInfo > FocusCamera;
 
-	/** Create a new blend profile */
-	TSharedPtr< FUICommandInfo > CreateBlendProfile;
+	/** Create a new blend profile with time mode*/
+	TSharedPtr< FUICommandInfo > CreateTimeBlendProfile;
+
+	/** Create a new blend profile with weight mode*/
+	TSharedPtr< FUICommandInfo > CreateWeightBlendProfile;
+
+	/** Create a new blend mask */
+	TSharedPtr< FUICommandInfo > CreateBlendMask;
 
 	/** Remove the currently active blend profile */
 	TSharedPtr< FUICommandInfo > DeleteCurrentBlendProfile;
