Bumping Persona Layout Versions

diff --git a/Engine/Source/Editor/AnimationBlueprintEditor/Private/AnimationBlueprintEditorMode.cpp b/Engine/Source/Editor/AnimationBlueprintEditor/Private/AnimationBlueprintEditorMode.cpp
--- a/Engine/Source/Editor/AnimationBlueprintEditor/Private/AnimationBlueprintEditorMode.cpp
+++ b/Engine/Source/Editor/AnimationBlueprintEditor/Private/AnimationBlueprintEditorMode.cpp
@@ -25,7 +25,7 @@ FAnimationBlueprintEditorMode::FAnimationBlueprintEditorMode(const TSharedRef<FA
 	PreviewScenePtr = InAnimationBlueprintEditor->GetPreviewScene();
 	AnimBlueprintPtr = CastChecked<UAnimBlueprint>(InAnimationBlueprintEditor->GetBlueprintObj());
 
-	TabLayout = FTabManager::NewLayout( "Stanalone_AnimationBlueprintEditMode_Layout_v1.3" )
+	TabLayout = FTabManager::NewLayout( "Stanalone_AnimationBlueprintEditMode_Layout_v1.4" )
 		->AddArea
 		(
 			FTabManager::NewPrimaryArea()
diff --git a/Engine/Source/Editor/AnimationEditor/Private/AnimationEditorMode.cpp b/Engine/Source/Editor/AnimationEditor/Private/AnimationEditorMode.cpp
--- a/Engine/Source/Editor/AnimationEditor/Private/AnimationEditorMode.cpp
+++ b/Engine/Source/Editor/AnimationEditor/Private/AnimationEditorMode.cpp
@@ -37,7 +37,7 @@ FAnimationEditorMode::FAnimationEditorMode(TSharedRef<FWorkflowCentricApplicatio
 	TabFactories.RegisterFactory(PersonaModule.CreateAnimNotifiesTabFactory(InHostingApp, InSkeletonTree->GetEditableSkeleton(), OnObjectsSelected));
 	TabFactories.RegisterFactory(PersonaModule.CreateAnimMontageSectionsTabFactory(InHostingApp, AnimationEditor->GetPersonaToolkit(), AnimationEditor->OnSectionsChanged));
 
-	TabLayout = FTabManager::NewLayout("Standalone_AnimationEditor_Layout_v1.3")
+	TabLayout = FTabManager::NewLayout("Standalone_AnimationEditor_Layout_v1.4")
 		->AddArea
 		(
 			FTabManager::NewPrimaryArea()
diff --git a/Engine/Source/Editor/PhysicsAssetEditor/Private/PhysicsAssetEditorMode.cpp b/Engine/Source/Editor/PhysicsAssetEditor/Private/PhysicsAssetEditorMode.cpp
--- a/Engine/Source/Editor/PhysicsAssetEditor/Private/PhysicsAssetEditorMode.cpp
+++ b/Engine/Source/Editor/PhysicsAssetEditor/Private/PhysicsAssetEditorMode.cpp
@@ -209,7 +209,7 @@ FPhysicsAssetEditorMode::FPhysicsAssetEditorMode(TSharedRef<FWorkflowCentricAppl
 
 	TabFactories.RegisterFactory(MakeShared<FPhysicsAssetEditorToolsSummoner>(InHostingApp));
 
-	TabLayout = FTabManager::NewLayout("Standalone_PhysicsAssetEditor_Layout_v5.2")
+	TabLayout = FTabManager::NewLayout("Standalone_PhysicsAssetEditor_Layout_v5.3")
 		->AddArea
 		(
 			FTabManager::NewPrimaryArea()
diff --git a/Engine/Source/Editor/SkeletalMeshEditor/Private/SkeletalMeshEditorMode.cpp b/Engine/Source/Editor/SkeletalMeshEditor/Private/SkeletalMeshEditorMode.cpp
--- a/Engine/Source/Editor/SkeletalMeshEditor/Private/SkeletalMeshEditorMode.cpp
+++ b/Engine/Source/Editor/SkeletalMeshEditor/Private/SkeletalMeshEditorMode.cpp
@@ -35,7 +35,7 @@ FSkeletalMeshEditorMode::FSkeletalMeshEditorMode(TSharedRef<FWorkflowCentricAppl
 
 	TabFactories.RegisterFactory(CreateMeshControllerMappingTabFactory(InHostingApp, Cast<USkeletalMesh> (SkeletalMeshEditor->HandleGetAsset()), SkeletalMeshEditor->OnPostUndo));
 
-	TabLayout = FTabManager::NewLayout("Standalone_SkeletalMeshEditor_Layout_v3.2")
+	TabLayout = FTabManager::NewLayout("Standalone_SkeletalMeshEditor_Layout_v3.3")
 		->AddArea
 		(
 			FTabManager::NewPrimaryArea()
diff --git a/Engine/Source/Editor/SkeletonEditor/Private/SkeletonEditorMode.cpp b/Engine/Source/Editor/SkeletonEditor/Private/SkeletonEditorMode.cpp
--- a/Engine/Source/Editor/SkeletonEditor/Private/SkeletonEditorMode.cpp
+++ b/Engine/Source/Editor/SkeletonEditor/Private/SkeletonEditorMode.cpp
@@ -35,7 +35,7 @@ FSkeletonEditorMode::FSkeletonEditorMode(TSharedRef<FWorkflowCentricApplication>
 	TabFactories.RegisterFactory(PersonaModule.CreateSkeletonSlotNamesTabFactory(InHostingApp, InSkeletonTree->GetEditableSkeleton(), SkeletonEditor->OnPostUndo, OnObjectSelected));
 	TabFactories.RegisterFactory(PersonaModule.CreateAssetDetailsTabFactory(InHostingApp, FOnGetAsset::CreateSP(&SkeletonEditor.Get(), &FSkeletonEditor::HandleGetAsset), FOnDetailsCreated()));
 
-	TabLayout = FTabManager::NewLayout("Standalone_SkeletonEditor_Layout_v1.2")
+	TabLayout = FTabManager::NewLayout("Standalone_SkeletonEditor_Layout_v1.3")
 		->AddArea
 		(
 			FTabManager::NewPrimaryArea()
