Changed slider range for weight based blend profiles, and allowed values outside the range (non-negative still disallowed) to be input manually.
#rnx

diff --git a/Engine/Source/Editor/SkeletonEditor/Private/SkeletonTreeBoneItem.cpp b/Engine/Source/Editor/SkeletonEditor/Private/SkeletonTreeBoneItem.cpp
--- a/Engine/Source/Editor/SkeletonEditor/Private/SkeletonTreeBoneItem.cpp
+++ b/Engine/Source/Editor/SkeletonEditor/Private/SkeletonTreeBoneItem.cpp
@@ -134,7 +134,8 @@ TSharedRef< SWidget > FSkeletonTreeBoneItem::GenerateWidgetForDataColumn(const F
 				.ContentPadding(0.0f)
 				.Delta(0.01f)
 				.MinValue(0.0f)
-				.MaxValue(1.0f)
+				.MinSliderValue(this, &FSkeletonTreeBoneItem::GetBlendProfileMinSliderValue)
+				.MaxSliderValue(this, &FSkeletonTreeBoneItem::GetBlendProfileMaxSliderValue)
 				.Value(this, &FSkeletonTreeBoneItem::GetBoneBlendProfileScale)
 				.OnValueCommitted(this, &FSkeletonTreeBoneItem::OnBlendSliderCommitted)
 				.OnValueChanged(this, &FSkeletonTreeBoneItem::OnBlendSliderCommitted, ETextCommit::OnEnter)
@@ -159,6 +160,26 @@ float FSkeletonTreeBoneItem::GetBoneBlendProfileScale()	const
 	return 0.0;
 }
 
+TOptional<float> FSkeletonTreeBoneItem::GetBlendProfileMaxSliderValue() const
+{
+	if (UBlendProfile* CurrentProfile = GetSkeletonTree()->GetSelectedBlendProfile())
+	{
+		return (CurrentProfile->GetMode() == EBlendProfileMode::WeightFactor) ? 10.0f : 1.0f;
+	}
+
+	return 1.0f;
+}
+
+TOptional<float> FSkeletonTreeBoneItem::GetBlendProfileMinSliderValue() const
+{
+	if (UBlendProfile* CurrentProfile = GetSkeletonTree()->GetSelectedBlendProfile())
+	{
+		return (CurrentProfile->GetMode() == EBlendProfileMode::WeightFactor) ? 1.0f : 0.0f;
+	}
+
+	return 0.0f;
+}
+
 FSlateColor FSkeletonTreeBoneItem::GetRetargetingComboButtonForegroundColor() const
 {
 	static const FName InvertedForegroundName("InvertedForeground");
diff --git a/Engine/Source/Editor/SkeletonEditor/Private/SkeletonTreeBoneItem.h b/Engine/Source/Editor/SkeletonEditor/Private/SkeletonTreeBoneItem.h
--- a/Engine/Source/Editor/SkeletonEditor/Private/SkeletonTreeBoneItem.h
+++ b/Engine/Source/Editor/SkeletonEditor/Private/SkeletonTreeBoneItem.h
@@ -85,6 +85,12 @@ private:
 	/** Get the current Blend Scale for this bone */
 	float GetBoneBlendProfileScale() const;
 
+	/** Get the max slider value supported by the selected blend profile's mode*/
+	TOptional<float> GetBlendProfileMaxSliderValue() const;
+
+	/** Get the min slider value supported by the selected blend profile's mode*/
+	TOptional<float> GetBlendProfileMinSliderValue() const;
+
 	/** Used to hide the Bone Profile row widget if there is no current Blend Profile */
 	EVisibility GetBoneBlendProfileVisibility() const;
 
