Physics Asset Editor: fix keyboard navigation in treeview
- make sure only visible (unfiltered) items can be navigated

diff --git a/Engine/Source/Editor/SkeletonEditor/Private/SSkeletonTree.cpp b/Engine/Source/Editor/SkeletonEditor/Private/SSkeletonTree.cpp
--- a/Engine/Source/Editor/SkeletonEditor/Private/SSkeletonTree.cpp
+++ b/Engine/Source/Editor/SkeletonEditor/Private/SSkeletonTree.cpp
@@ -631,6 +631,7 @@ void SSkeletonTree::CreateTreeColumns()
 		.OnGetChildren(this, &SSkeletonTree::GetFilteredChildren)
 		.OnContextMenuOpening(this, &SSkeletonTree::CreateContextMenu)
 		.OnSelectionChanged(this, &SSkeletonTree::OnSelectionChanged)
+		.OnIsSelectableOrNavigable(this, &SSkeletonTree::OnIsSelectableOrNavigable)
 		.OnItemScrolledIntoView(this, &SSkeletonTree::OnItemScrolledIntoView)
 		.OnMouseButtonDoubleClick(this, &SSkeletonTree::OnTreeDoubleClick)
 		.OnSetExpansionRecursive(this, &SSkeletonTree::SetTreeItemExpansionRecursive)
@@ -1398,6 +1399,11 @@ void SSkeletonTree::OnRenameSelected()
 	}
 }
 
+bool SSkeletonTree::OnIsSelectableOrNavigable(TSharedPtr<class ISkeletonTreeItem> InItem) const
+{
+	return InItem && InItem->GetFilterResult() == ESkeletonTreeFilterResult::Shown;
+}
+
 void SSkeletonTree::OnSelectionChanged(TSharedPtr<ISkeletonTreeItem> Selection, ESelectInfo::Type SelectInfo)
 {
 	TArray<TSharedPtr<ISkeletonTreeItem>> SelectedItems = SkeletonTreeView->GetSelectedItems();
diff --git a/Engine/Source/Editor/SkeletonEditor/Private/SSkeletonTree.h b/Engine/Source/Editor/SkeletonEditor/Private/SSkeletonTree.h
--- a/Engine/Source/Editor/SkeletonEditor/Private/SSkeletonTree.h
+++ b/Engine/Source/Editor/SkeletonEditor/Private/SSkeletonTree.h
@@ -240,6 +240,9 @@ private:
 	/** Functions to copy sockets from the skeleton to the mesh */
 	void OnCopySocketToMesh() {};
 
+	/** Callback function to be called when selection changes to check if the next item is selectable or navigable. */
+	bool OnIsSelectableOrNavigable(TSharedPtr<class ISkeletonTreeItem> InItem) const;
+
 	/** Callback function to be called when selection changes in the tree view widget. */
 	void OnSelectionChanged(TSharedPtr<class ISkeletonTreeItem> Selection, ESelectInfo::Type SelectInfo);
 
