Undo/Redo operations not working on skeleton asset
#fix ensure that during PostUndo/Redo the skeleton tree is rebuild, behaviour was lost during restyle refactor

diff --git a/Engine/Source/Editor/SkeletonEditor/Private/SSkeletonTree.cpp b/Engine/Source/Editor/SkeletonEditor/Private/SSkeletonTree.cpp
--- a/Engine/Source/Editor/SkeletonEditor/Private/SSkeletonTree.cpp
+++ b/Engine/Source/Editor/SkeletonEditor/Private/SSkeletonTree.cpp
@@ -1566,12 +1566,14 @@ void SSkeletonTree::SetTreeItemExpansionRecursive(TSharedPtr< ISkeletonTreeItem
 void SSkeletonTree::PostUndo(bool bSuccess)
 {
 	// Rebuild the tree view whenever we undo a change to the skeleton
+	CreateFromSkeleton();
 	HandleTreeRefresh();
 }
 
 void SSkeletonTree::PostRedo(bool bSuccess)
 {
 	// Rebuild the tree view whenever we redo a change to the skeleton
+	CreateFromSkeleton();
 	HandleTreeRefresh();
 }
 
