fix get references to skeleton notify from skeleton editor

diff --git a/Engine/Source/Editor/Persona/Private/SSkeletonAnimNotifies.cpp b/Engine/Source/Editor/Persona/Private/SSkeletonAnimNotifies.cpp
--- a/Engine/Source/Editor/Persona/Private/SSkeletonAnimNotifies.cpp
+++ b/Engine/Source/Editor/Persona/Private/SSkeletonAnimNotifies.cpp
@@ -19,6 +19,7 @@
 #include "IAnimationEditor.h"
 #include "IAnimationBlueprintEditor.h"
 #include "IAnimationSequenceBrowser.h"
+#include "ISkeletonEditor.h"
 
 #define LOCTEXT_NAMESPACE "SkeletonAnimNotifies"
 
@@ -496,6 +497,11 @@ void SSkeletonAnimNotifies::OnFindReferences()
 			TSharedPtr<IAnimationBlueprintEditor> AnimationBlueprintEditor = StaticCastSharedPtr<IAnimationBlueprintEditor>(HostingApp);
 			AnimationBlueprintEditor->GetAssetBrowser()->FilterBySkeletonNotify(Name);
 		}
+		else if (HostingApp->GetEditorName() == TEXT("SkeletonEditor"))
+		{
+			TSharedPtr<ISkeletonEditor> SkeletonEditor = StaticCastSharedPtr<ISkeletonEditor>(HostingApp);
+			SkeletonEditor->GetAssetBrowser()->FilterBySkeletonNotify(Name);
+		}
 	}
 }
 
diff --git a/Engine/Source/Editor/SkeletonEditor/Private/SkeletonEditor.cpp b/Engine/Source/Editor/SkeletonEditor/Private/SkeletonEditor.cpp
--- a/Engine/Source/Editor/SkeletonEditor/Private/SkeletonEditor.cpp
+++ b/Engine/Source/Editor/SkeletonEditor/Private/SkeletonEditor.cpp
@@ -14,6 +14,7 @@
 #include "IPersonaToolkit.h"
 #include "PersonaModule.h"
 #include "SkeletonEditorMode.h"
+#include "IAnimationSequenceBrowser.h"
 #include "IPersonaPreviewScene.h"
 #include "SkeletonEditorCommands.h"
 #include "IAssetFamily.h"
@@ -30,6 +31,7 @@ const FName SkeletonEditorModes::SkeletonEditorMode(TEXT("SkeletonEditorMode"));
 const FName SkeletonEditorTabs::DetailsTab(TEXT("DetailsTab"));
 const FName SkeletonEditorTabs::SkeletonTreeTab(TEXT("SkeletonTreeView"));
 const FName SkeletonEditorTabs::ViewportTab(TEXT("Viewport"));
+const FName SkeletonEditorTabs::AssetBrowserTab(TEXT("SequenceBrowser"));
 const FName SkeletonEditorTabs::AnimNotifiesTab(TEXT("SkeletonAnimNotifies"));
 const FName SkeletonEditorTabs::CurveNamesTab(TEXT("AnimCurveViewerTab"));
 const FName SkeletonEditorTabs::AdvancedPreviewTab(TEXT("AdvancedPreviewTab"));
@@ -58,6 +60,11 @@ FSkeletonEditor::~FSkeletonEditor()
 	}
 }
 
+void FSkeletonEditor::HandleOpenNewAsset(UObject* InNewAsset)
+{
+	GEditor->GetEditorSubsystem<UAssetEditorSubsystem>()->OpenEditorForAsset(InNewAsset);
+}
+
 void FSkeletonEditor::RegisterTabSpawners(const TSharedRef<class FTabManager>& InTabManager)
 {
 	WorkspaceMenuCategory = InTabManager->AddLocalWorkspaceMenuCategory(LOCTEXT("WorkspaceMenu_SkeletonEditor", "Skeleton Editor"));
@@ -342,4 +349,14 @@ void FSkeletonEditor::HandleDetailsCreated(const TSharedRef<IDetailsView>& InDet
 	DetailsView = InDetailsView;
 }
 
+void FSkeletonEditor::HandleAnimationSequenceBrowserCreated(const TSharedRef<IAnimationSequenceBrowser>& InSequenceBrowser)
+{
+	SequenceBrowser = InSequenceBrowser;
+}
+
+IAnimationSequenceBrowser* FSkeletonEditor::GetAssetBrowser() const
+{
+	return SequenceBrowser.Pin().Get();
+}
+
 #undef LOCTEXT_NAMESPACE
diff --git a/Engine/Source/Editor/SkeletonEditor/Private/SkeletonEditor.h b/Engine/Source/Editor/SkeletonEditor/Private/SkeletonEditor.h
--- a/Engine/Source/Editor/SkeletonEditor/Private/SkeletonEditor.h
+++ b/Engine/Source/Editor/SkeletonEditor/Private/SkeletonEditor.h
@@ -10,6 +10,7 @@
 #include "IDetailsView.h"
 #include "ISkeletonEditor.h"
 #include "Containers/ArrayView.h"
+#include "IAnimationSequenceBrowser.h"
 
 class IPersonaToolkit;
 class IPersonaViewport;
@@ -31,6 +32,7 @@ namespace SkeletonEditorTabs
 	extern const FName DetailsTab;
 	extern const FName SkeletonTreeTab;
 	extern const FName ViewportTab;
+	extern const FName AssetBrowserTab;
 	extern const FName AnimNotifiesTab;
 	extern const FName CurveNamesTab;
 	extern const FName AdvancedPreviewTab;
@@ -60,6 +62,11 @@ public:
 	virtual FLinearColor GetWorldCentricTabColorScale() const override;
 	virtual void InitToolMenuContext(FToolMenuContext& MenuContext) override;
 
+	virtual IAnimationSequenceBrowser* GetAssetBrowser() const override;
+
+	/** Handle opening a new asset from the asset browser */
+	void HandleOpenNewAsset(UObject* InNewAsset);
+
 	/** FEditorUndoClient interface */
 	virtual void PostUndo(bool bSuccess) override;
 	virtual void PostRedo(bool bSuccess) override;
@@ -89,6 +96,7 @@ public:
 
 	void HandleDetailsCreated(const TSharedRef<class IDetailsView>& InDetailsView);
 
+	void HandleAnimationSequenceBrowserCreated(const TSharedRef<class IAnimationSequenceBrowser>& InSequenceBrowser);
 private:
 	void ExtendMenu();
 
@@ -138,4 +146,7 @@ private:
 
 	/** Details panel */
 	TSharedPtr<class IDetailsView> DetailsView;
+
+	/** Sequence Browser **/
+	TWeakPtr<class IAnimationSequenceBrowser> SequenceBrowser;
 };
diff --git a/Engine/Source/Editor/SkeletonEditor/Private/SkeletonEditorMode.cpp b/Engine/Source/Editor/SkeletonEditor/Private/SkeletonEditorMode.cpp
--- a/Engine/Source/Editor/SkeletonEditor/Private/SkeletonEditorMode.cpp
+++ b/Engine/Source/Editor/SkeletonEditor/Private/SkeletonEditorMode.cpp
@@ -34,6 +34,7 @@ FSkeletonEditorMode::FSkeletonEditorMode(TSharedRef<FWorkflowCentricApplication>
 	TabFactories.RegisterFactory(PersonaModule.CreateCurveViewerTabFactory(InHostingApp, InSkeletonTree->GetEditableSkeleton(), SkeletonEditor->GetPersonaToolkit()->GetPreviewScene(), OnObjectsSelected));
 	TabFactories.RegisterFactory(PersonaModule.CreateSkeletonSlotNamesTabFactory(InHostingApp, InSkeletonTree->GetEditableSkeleton(), OnObjectSelected));
 	TabFactories.RegisterFactory(PersonaModule.CreateAssetDetailsTabFactory(InHostingApp, FOnGetAsset::CreateSP(&SkeletonEditor.Get(), &FSkeletonEditor::HandleGetAsset), FOnDetailsCreated()));
+	TabFactories.RegisterFactory(PersonaModule.CreateAnimationAssetBrowserTabFactory(InHostingApp, SkeletonEditor->GetPersonaToolkit(), FOnOpenNewAsset::CreateSP(&SkeletonEditor.Get(), &FSkeletonEditor::HandleOpenNewAsset), FOnAnimationSequenceBrowserCreated::CreateSP(&SkeletonEditor.Get(), &FSkeletonEditor::HandleAnimationSequenceBrowserCreated), true));
 
 	TabLayout = FTabManager::NewLayout("Standalone_SkeletonEditor_Layout_v1.3")
 		->AddArea
@@ -80,6 +81,7 @@ FSkeletonEditorMode::FSkeletonEditorMode(TSharedRef<FWorkflowCentricApplication>
 						->SetSizeCoefficient(0.5f)
 						->SetHideTabWell(false)
 						->AddTab(SkeletonEditorTabs::AnimNotifiesTab, ETabState::OpenedTab)
+						->AddTab(SkeletonEditorTabs::AssetBrowserTab, ETabState::OpenedTab)
 						->AddTab(SkeletonEditorTabs::CurveNamesTab, ETabState::OpenedTab)
 						->AddTab(SkeletonEditorTabs::SlotNamesTab, ETabState::ClosedTab)
 					)
diff --git a/Engine/Source/Editor/SkeletonEditor/Public/ISkeletonEditor.h b/Engine/Source/Editor/SkeletonEditor/Public/ISkeletonEditor.h
--- a/Engine/Source/Editor/SkeletonEditor/Public/ISkeletonEditor.h
+++ b/Engine/Source/Editor/SkeletonEditor/Public/ISkeletonEditor.h
@@ -6,6 +6,11 @@
 #include "PersonaAssetEditorToolkit.h"
 #include "IHasPersonaToolkit.h"
 
+class IAnimationSequenceBrowser;
+
 class ISkeletonEditor : public FPersonaAssetEditorToolkit, public IHasPersonaToolkit
 {
+	public:
+	/** Get the asset browser we host */
+	virtual IAnimationSequenceBrowser* GetAssetBrowser() const = 0;
 };
