Forces showing Blend Profile Column when selecting a blend profile

diff --git a/Engine/Source/Editor/SkeletonEditor/Private/SSkeletonTree.cpp b/Engine/Source/Editor/SkeletonEditor/Private/SSkeletonTree.cpp
--- a/Engine/Source/Editor/SkeletonEditor/Private/SSkeletonTree.cpp
+++ b/Engine/Source/Editor/SkeletonEditor/Private/SSkeletonTree.cpp
@@ -2091,6 +2091,8 @@ FName SSkeletonTree::GetSelectedBlendProfileName() const
 void SSkeletonTree::OnBlendProfileSelected(UBlendProfile* NewProfile)
 {
 	SkeletonTreeView->GetHeaderRow()->RefreshColumns();
+	if (NewProfile != nullptr)
+		SkeletonTreeView->GetHeaderRow()->SetShowGeneratedColumn(ISkeletonTree::Columns::BlendProfile);
 	HandleTreeRefresh();
 
 	// When a new blend profile is created/selected - reaffirm that the header can't be edited.
