SkeletalMesh Retargeting Settings Fix

diff --git a/Engine/Source/Editor/SkeletonEditor/Private/SSkeletonTree.cpp b/Engine/Source/Editor/SkeletonEditor/Private/SSkeletonTree.cpp
--- a/Engine/Source/Editor/SkeletonEditor/Private/SSkeletonTree.cpp
+++ b/Engine/Source/Editor/SkeletonEditor/Private/SSkeletonTree.cpp
@@ -132,7 +132,6 @@ void SSkeletonTree::Construct(const FArguments& InArgs, const TSharedRef<FEditab
 		BoneFilter = EBoneFilter::All;
 	}
 	SocketFilter = ESocketFilter::Active;
-	bShowingAdvancedOptions = false;
 	bSelecting = false;
 
 	EditableSkeleton = InEditableSkeleton;
@@ -570,16 +569,6 @@ struct FScopedSavedSelection
 
 void SSkeletonTree::CreateTreeColumns()
 {
-	auto HeaderVisibilityLambda = [this]()
-	{
-		if(Mode == ESkeletonTreeMode::Editor)
-		{
-			return bShowingAdvancedOptions || BlendProfilePicker->GetSelectedBlendProfileName() != NAME_None ? EVisibility::Visible : EVisibility::Collapsed;
-		}
-		
-		return  EVisibility::Collapsed; 
-	};
-
 	TArray<FName> HiddenColumnsList;
 	HiddenColumnsList.Add(ISkeletonTree::Columns::Retargeting);
 	HiddenColumnsList.Add(ISkeletonTree::Columns::BlendProfile);
@@ -786,7 +775,7 @@ TSharedPtr< SWidget > SSkeletonTree::CreateContextMenu()
 					MenuBuilder.EndSection();
 				}
 
-				if(bShowingAdvancedOptions)
+				if(IsShowingAdvancedOptions())
 				{
 					MenuBuilder.BeginSection("SkeletonTreeBoneTranslationRetargeting", LOCTEXT("BoneTranslationRetargetingHeader", "Bone Translation Retargeting"));
 					{
@@ -2069,13 +2058,13 @@ void SSkeletonTree::DeleteVirtualBones(const TArray<TSharedPtr<FSkeletonTreeVirt
 
 void SSkeletonTree::OnChangeShowingAdvancedOptions()
 {
-	bShowingAdvancedOptions = !bShowingAdvancedOptions;
+	SkeletonTreeView->GetHeaderRow()->SetShowGeneratedColumn(ISkeletonTree::Columns::Retargeting, !IsShowingAdvancedOptions());
 	HandleTreeRefresh();
 }
 
 bool SSkeletonTree::IsShowingAdvancedOptions() const
 {
-	return bShowingAdvancedOptions;
+	return SkeletonTreeView->GetHeaderRow()->IsColumnVisible(ISkeletonTree::Columns::Retargeting);
 }
 
 UBlendProfile* SSkeletonTree::GetSelectedBlendProfile()
diff --git a/Engine/Source/Editor/SkeletonEditor/Private/SSkeletonTree.h b/Engine/Source/Editor/SkeletonEditor/Private/SSkeletonTree.h
--- a/Engine/Source/Editor/SkeletonEditor/Private/SSkeletonTree.h
+++ b/Engine/Source/Editor/SkeletonEditor/Private/SSkeletonTree.h
@@ -387,8 +387,6 @@ private:
 	/** Current type of sockets to show */
 	ESocketFilter SocketFilter;
 
-	bool bShowingAdvancedOptions;
-
 	/** Points to an item that is being requested to be renamed */
 	TSharedPtr<ISkeletonTreeItem> DeferredRenameRequest;
 
diff --git a/Engine/Source/Runtime/Slate/Private/Widgets/Views/SHeaderRow.cpp b/Engine/Source/Runtime/Slate/Private/Widgets/Views/SHeaderRow.cpp
--- a/Engine/Source/Runtime/Slate/Private/Widgets/Views/SHeaderRow.cpp
+++ b/Engine/Source/Runtime/Slate/Private/Widgets/Views/SHeaderRow.cpp
@@ -642,6 +642,18 @@ bool SHeaderRow::IsColumnGenerated(const FName& InColumnId) const
 	return false;
 }
 
+bool SHeaderRow::IsColumnVisible(const FName& InColumnId) const
+{
+	for (const FColumn& SomeColumn : Columns)
+	{
+		if (SomeColumn.ColumnId == InColumnId)
+		{
+			return SomeColumn.bIsVisible;
+		}
+	}
+	return false;
+}
+
 FReply SHeaderRow::OnMouseButtonUp(const FGeometry& MyGeometry, const FPointerEvent& MouseEvent)
 {
 	if (bCanSelectGeneratedColumn && MouseEvent.GetEffectingButton() == EKeys::RightMouseButton)
diff --git a/Engine/Source/Runtime/Slate/Public/Widgets/Views/SHeaderRow.h b/Engine/Source/Runtime/Slate/Public/Widgets/Views/SHeaderRow.h
--- a/Engine/Source/Runtime/Slate/Public/Widgets/Views/SHeaderRow.h
+++ b/Engine/Source/Runtime/Slate/Public/Widgets/Views/SHeaderRow.h
@@ -337,6 +337,9 @@ public:
 	/** @return is the header widget with the column id is generated */
 	bool IsColumnGenerated( const FName& InColumnId ) const;
 
+	/** @return true if the header widget with the column id is visible */
+	bool IsColumnVisible( const FName& InColumnId ) const;
+
 	//~ Begin SWidget interface
 	virtual FReply OnMouseButtonUp( const FGeometry& MyGeometry, const FPointerEvent& MouseEvent ) override;
 	//~ End SWidget interface
