Animation Tree UI: Introduce option to expand tree based on selection

diff --git a/Engine/Plugins/Animation/ControlRig/Source/ControlRigEditor/Private/Editor/SRigHierarchy.cpp b/Engine/Plugins/Animation/ControlRig/Source/ControlRigEditor/Private/Editor/SRigHierarchy.cpp
--- a/Engine/Plugins/Animation/ControlRig/Source/ControlRigEditor/Private/Editor/SRigHierarchy.cpp
+++ b/Engine/Plugins/Animation/ControlRig/Source/ControlRigEditor/Private/Editor/SRigHierarchy.cpp
@@ -713,7 +713,11 @@ void SRigHierarchy::OnHierarchyModified(ERigHierarchyNotification InNotif, URigH
 					if (Found.IsValid())
 					{
 						TreeView->SetItemSelection(Found, bSelected, ESelectInfo::OnNavigation);
-						HandleFrameSelection();
+
+						if(GetDefault<UPersonaOptions>()->bExpandTreeOnSelection && bSelected)
+						{
+							HandleFrameSelection();
+						}
 
 						if (ControlRigEditor.IsValid() && !GIsTransacting)
 						{
diff --git a/Engine/Plugins/Animation/IKRig/Source/IKRigEditor/Private/RigEditor/SIKRigSkeleton.cpp b/Engine/Plugins/Animation/IKRig/Source/IKRigEditor/Private/RigEditor/SIKRigSkeleton.cpp
--- a/Engine/Plugins/Animation/IKRig/Source/IKRigEditor/Private/RigEditor/SIKRigSkeleton.cpp
+++ b/Engine/Plugins/Animation/IKRig/Source/IKRigEditor/Private/RigEditor/SIKRigSkeleton.cpp
@@ -16,6 +16,7 @@
 #include "Widgets/Text/SInlineEditableTextBlock.h"
 #include "Framework/MultiBox/MultiBoxBuilder.h"
 #include "Framework/Commands/UICommandList.h"
+#include "Preferences/PersonaOptions.h"
 
 #define LOCTEXT_NAMESPACE "SIKRigSkeleton"
 
@@ -274,6 +275,18 @@ void SIKRigSkeleton::SetSelectedGoalsFromViewport(const TArray<FName>& GoalNames
 		if (GoalNames.Contains(Item->GoalName))
 		{
 			TreeView->SetSelection(Item, ESelectInfo::Direct);
+
+			if(GetDefault<UPersonaOptions>()->bExpandTreeOnSelection)
+			{
+				TSharedPtr<FIKRigTreeElement> ItemToExpand = Item->Parent;
+				while(ItemToExpand.IsValid())
+				{
+					TreeView->SetItemExpansion(ItemToExpand, true);
+					ItemToExpand = ItemToExpand->Parent;
+				}
+			}
+
+			TreeView->RequestScrollIntoView(Item);
 		}
 	}
 }
diff --git a/Engine/Source/Editor/SkeletonEditor/Private/SSkeletonTree.cpp b/Engine/Source/Editor/SkeletonEditor/Private/SSkeletonTree.cpp
--- a/Engine/Source/Editor/SkeletonEditor/Private/SSkeletonTree.cpp
+++ b/Engine/Source/Editor/SkeletonEditor/Private/SSkeletonTree.cpp
@@ -1599,6 +1599,19 @@ TSharedRef<SWidget> SSkeletonTree::GetBlendProfileColumnMenuContent()
 	return MenuBuilder.MakeWidget();
 }
 
+void SSkeletonTree::ExpandTreeOnSelection(TSharedPtr<ISkeletonTreeItem> RowToExpand, bool bForce)
+{
+	if(GetDefault<UPersonaOptions>()->bExpandTreeOnSelection || bForce)
+	{
+		RowToExpand = RowToExpand->GetParent();
+		while(RowToExpand.IsValid())
+		{
+			SkeletonTreeView->SetItemExpansion(RowToExpand, true);
+			RowToExpand = RowToExpand->GetParent();
+		}
+	}
+}
+
 void SSkeletonTree::GetBlendProfileMenu(FMenuBuilder& MenuBuilder)
 {
 	const FSkeletonTreeCommands& Actions = FSkeletonTreeCommands::Get();
@@ -1701,6 +1714,7 @@ void SSkeletonTree::SetSelectedSocket( const FSelectedSocketInfo& SocketInfo )
 			if (SkeletonRow->GetFilterResult() != ESkeletonTreeFilterResult::Hidden && SkeletonRow->IsOfType<FSkeletonTreeSocketItem>() && StaticCastSharedPtr<FSkeletonTreeSocketItem>(SkeletonRow)->GetSocket() == SocketInfo.Socket)
 			{
 				SkeletonTreeView->SetItemSelection(SkeletonRow, true);
+				ExpandTreeOnSelection(SkeletonRow);
 				SkeletonTreeView->RequestScrollIntoView(SkeletonRow);
 			}
 		}
@@ -1725,6 +1739,7 @@ void SSkeletonTree::SetSelectedBone( const FName& BoneName, ESelectInfo::Type In
 			if (SkeletonRow->GetFilterResult() != ESkeletonTreeFilterResult::Hidden && SkeletonRow->IsOfType<FSkeletonTreeBoneItem>() && SkeletonRow->GetRowItemName() == BoneName)
 			{
 				SkeletonTreeView->SetItemSelection(SkeletonRow, true, InSelectInfo);
+				ExpandTreeOnSelection(SkeletonRow);
 				SkeletonTreeView->RequestScrollIntoView(SkeletonRow);
 			}
 		}
@@ -2259,6 +2274,13 @@ void SSkeletonTree::HandleFocusCamera()
 	{
 		GetPreviewScene()->FocusViews();
 	}
+
+	if(!SkeletonTreeView->GetSelectedItems().IsEmpty())
+	{
+		TSharedPtr<class ISkeletonTreeItem> SelectedRow = SkeletonTreeView->GetSelectedItems()[0]; 
+		ExpandTreeOnSelection(SelectedRow);
+		SkeletonTreeView->RequestScrollIntoView(SelectedRow);
+	}
 }
 
 ESkeletonTreeFilterResult SSkeletonTree::HandleFilterSkeletonTreeItem(const FSkeletonTreeFilterArgs& InArgs, const TSharedPtr<class ISkeletonTreeItem>& InItem)
diff --git a/Engine/Source/Editor/SkeletonEditor/Private/SSkeletonTree.h b/Engine/Source/Editor/SkeletonEditor/Private/SSkeletonTree.h
--- a/Engine/Source/Editor/SkeletonEditor/Private/SSkeletonTree.h
+++ b/Engine/Source/Editor/SkeletonEditor/Private/SSkeletonTree.h
@@ -313,6 +313,9 @@ private:
 	/** Handle focusing the camera on the current selection */
 	void HandleFocusCamera();
 
+	/** Handle framing the selected item in the tree */
+	void HandleFrameSelection();
+
 	/** Handle filtering the tree  */
 	ESkeletonTreeFilterResult HandleFilterSkeletonTreeItem(const FSkeletonTreeFilterArgs& InArgs, const TSharedPtr<class ISkeletonTreeItem>& InItem);
 
@@ -345,6 +348,8 @@ private:
 
 	TSharedRef<SWidget> GetBlendProfileColumnMenuContent();
 
+	void ExpandTreeOnSelection(TSharedPtr<ISkeletonTreeItem> RowToExpand, bool bForce = false);
+
 private:
 	/** Pointer back to the skeleton tree that owns us */
 	TWeakPtr<FEditableSkeleton> EditableSkeleton;
diff --git a/Engine/Source/Editor/UnrealEd/Classes/Preferences/PersonaOptions.h b/Engine/Source/Editor/UnrealEd/Classes/Preferences/PersonaOptions.h
--- a/Engine/Source/Editor/UnrealEd/Classes/Preferences/PersonaOptions.h
+++ b/Engine/Source/Editor/UnrealEd/Classes/Preferences/PersonaOptions.h
@@ -163,6 +163,10 @@ class UNREALED_API UPersonaOptions : public UObject
 	UPROPERTY(EditAnywhere, config, Category = "Skeleton Tree")
 	bool bHideParentsWhenFiltering;
 
+	/** Whether to focus and expand an item's tree recursively based on selection */
+	UPROPERTY(EditAnywhere, config, Category = "Skeleton Tree")
+	bool bExpandTreeOnSelection;
+
 	UPROPERTY(EditAnywhere, config, Category = "Preview Scene|AdditionalMesh")
 	bool bAllowPreviewMeshCollectionsToSelectFromDifferentSkeletons;
 
diff --git a/Engine/Source/Editor/UnrealEd/Private/PreferenceStubs.cpp b/Engine/Source/Editor/UnrealEd/Private/PreferenceStubs.cpp
--- a/Engine/Source/Editor/UnrealEd/Private/PreferenceStubs.cpp
+++ b/Engine/Source/Editor/UnrealEd/Private/PreferenceStubs.cpp
@@ -142,6 +142,8 @@ UPersonaOptions::UPersonaOptions(const FObjectInitializer& ObjectInitializer)
 	bTimelineDisplayCurveKeys = false;
 
 	TimelineEnabledSnaps = { "CompositeSegment", "MontageSection" };
+
+	bExpandTreeOnSelection = true;
 }
 
 void UPersonaOptions::SetShowGrid( bool bInShowGrid )
