Coalition skeleton tree item sorting

diff --git a/Engine/Source/Editor/SkeletonEditor/Private/SkeletonTreeBuilder.cpp b/Engine/Source/Editor/SkeletonEditor/Private/SkeletonTreeBuilder.cpp
--- a/Engine/Source/Editor/SkeletonEditor/Private/SkeletonTreeBuilder.cpp
+++ b/Engine/Source/Editor/SkeletonEditor/Private/SkeletonTreeBuilder.cpp
@@ -208,18 +208,88 @@ void FSkeletonTreeBuilder::AddBones(FSkeletonTreeBuilderOutput& Output)
 	const USkeleton& Skeleton = EditableSkeletonPtr.Pin()->GetSkeleton();
 
 	const FReferenceSkeleton& RefSkeleton = Skeleton.GetReferenceSkeleton();
+
+	struct FBoneInfo
+	{
+		FBoneInfo(const FName& InBoneName, const FName& InParentName, int32 InDepth)
+			: BoneName(InBoneName)
+			, ParentName(InParentName)
+			, Depth(InDepth)
+		{
+			SortString = BoneName.ToString();
+			SortNumber = 0;
+			SortLength = SortString.Len();
+
+			// Split the bone name into string prefix and numeric suffix for sorting (different from FName to support leading zeros in the numeric suffix)
+			int32 Index = SortLength - 1;
+			for (int32 PlaceValue = 1; Index >= 0 && FChar::IsDigit(SortString[Index]); --Index, PlaceValue *= 10)
+			{
+				SortNumber += static_cast<int32>(SortString[Index] - '0') * PlaceValue;
+			}
+			SortString.LeftInline(Index + 1, false);
+		}
+
+		bool operator<(const FBoneInfo& RHS)
+		{
+			// Sort parents before children
+			if (Depth != RHS.Depth)
+			{
+				return Depth < RHS.Depth;
+			}
+
+			// Sort alphabetically by string prefix
+			if (int32 SplitNameComparison = SortString.Compare(RHS.SortString))
+			{
+				return SplitNameComparison < 0;
+			}
+
+			// Sort by number if the string prefixes match
+			if (SortNumber != RHS.SortNumber)
+			{
+				return SortNumber < RHS.SortNumber;
+			}
+
+			// Sort by length to give us the equivalent to alphabetical sorting if the numbers match (which gives us the following sort order: bone_, bone_0, bone_00, bone_000, bone_001, bone_01, bone_1, etc)
+			return (SortNumber == 0) ? SortLength < RHS.SortLength : SortLength > RHS.SortLength;
+		}
+
+		FName BoneName = NAME_None;
+		FName ParentName = NAME_None;
+		int32 Depth = 0;
+
+		FString SortString;
+		int32 SortNumber = 0;
+		int32 SortLength = 0;
+	};
+
+	TArray<FBoneInfo> Bones;
+	Bones.Reserve(RefSkeleton.GetRawBoneNum());
+
+	// Gather the bones from the skeleton
 	for (int32 BoneIndex = 0; BoneIndex < RefSkeleton.GetRawBoneNum(); ++BoneIndex)
 	{
 		const FName& BoneName = RefSkeleton.GetBoneName(BoneIndex);
 
 		FName ParentName = NAME_None;
 		int32 ParentIndex = RefSkeleton.GetParentIndex(BoneIndex);
+		int32 Depth = 0;
 		if (ParentIndex != INDEX_NONE)
 		{
 			ParentName = RefSkeleton.GetBoneName(ParentIndex);
+			Depth = Bones[ParentIndex].Depth + 1;
 		}
-		TSharedRef<ISkeletonTreeItem> DisplayBone = CreateBoneTreeItem(BoneName);
-		Output.Add(DisplayBone, ParentName, FSkeletonTreeBoneItem::GetTypeId());
+
+		Bones.Emplace(BoneName, ParentName, Depth);
+	}
+
+	// Sort the bones lexically (and also by depth in order to maintain the invariant of parents before children)
+	Algo::Sort(Bones);
+
+	// Add the sorted bones to the skeleton tree
+	for (const FBoneInfo& Bone : Bones)
+	{
+		TSharedRef<ISkeletonTreeItem> DisplayBone = CreateBoneTreeItem(Bone.BoneName);
+		Output.Add(DisplayBone, Bone.ParentName, FSkeletonTreeBoneItem::GetTypeId());
 	}
 }
 
