SkeletonEditor: Use new SAdvancedTransformInputBox

diff --git a/Engine/Source/Editor/SkeletonEditor/Private/BoneProxy.cpp b/Engine/Source/Editor/SkeletonEditor/Private/BoneProxy.cpp
--- a/Engine/Source/Editor/SkeletonEditor/Private/BoneProxy.cpp
+++ b/Engine/Source/Editor/SkeletonEditor/Private/BoneProxy.cpp
@@ -12,6 +12,7 @@
 UBoneProxy::UBoneProxy()
 	: bLocalLocation(true)
 	, bLocalRotation(true)
+	, bLocalScale(true)
 	, PreviousLocation(FVector::ZeroVector)
 	, PreviousRotation(FRotator::ZeroRotator)
 	, PreviousScale(FVector::ZeroVector)
@@ -52,7 +53,14 @@ void UBoneProxy::Tick(float DeltaTime)
 					Rotation = BoneTransform.GetRotation().Rotator();
 				}
 
-				Scale = LocalTransform.GetScale3D();
+				if(bLocalScale)
+				{
+					Scale = LocalTransform.GetScale3D();
+				}
+				else
+				{
+					Scale = BoneTransform.GetScale3D();
+				}
 
 				FTransform ReferenceTransform = Component->SkeletalMesh->GetRefSkeleton().GetRefBonePose()[BoneIndex];
 				ReferenceLocation = ReferenceTransform.GetLocation();
@@ -82,6 +90,220 @@ TStatId UBoneProxy::GetStatId() const
 	RETURN_QUICK_DECLARE_CYCLE_STAT(UBoneProxy, STATGROUP_Tickables);
 }
 
+TOptional<FVector::FReal> UBoneProxy::GetNumericValue(
+	ESlateTransformComponent::Type Component,
+	ESlateRotationRepresentation::Type Representation,
+	ESlateTransformSubComponent::Type SubComponent,
+	ETransformType TransformType) const
+{
+	const FVector* LocationPtr = &Location;
+	const FRotator* RotationPtr = &Rotation;
+	const FVector* ScalePtr = &Scale;
+
+	switch(TransformType)
+	{
+		case TransformType_Reference:
+		{
+			LocationPtr = &ReferenceLocation;
+			RotationPtr = &ReferenceRotation;
+			ScalePtr = &ReferenceScale;
+			break;
+		}
+		case TransformType_Mesh:
+		{
+			LocationPtr = &MeshLocation;
+			RotationPtr = &MeshRotation;
+			ScalePtr = &MeshScale;
+			break;
+		}
+	}
+	
+	switch(Component)
+	{
+		case ESlateTransformComponent::Location:
+		{
+			switch(SubComponent)
+			{
+				case ESlateTransformSubComponent::X:
+					return LocationPtr->X;
+				case ESlateTransformSubComponent::Y:
+					return LocationPtr->Y;
+				case ESlateTransformSubComponent::Z:
+					return LocationPtr->Z;
+			}
+			break;
+		}
+		case ESlateTransformComponent::Rotation:
+		{
+			switch(SubComponent)
+			{
+			case ESlateTransformSubComponent::Roll:
+				return RotationPtr->Roll;
+			case ESlateTransformSubComponent::Pitch:
+				return RotationPtr->Pitch;
+			case ESlateTransformSubComponent::Yaw:
+				return RotationPtr->Yaw;
+			}
+			break;
+		}
+		case ESlateTransformComponent::Scale:
+		{
+			switch(SubComponent)
+			{
+			case ESlateTransformSubComponent::X:
+				return ScalePtr->X;
+			case ESlateTransformSubComponent::Y:
+				return ScalePtr->Y;
+			case ESlateTransformSubComponent::Z:
+				return ScalePtr->Z;
+			}
+			break;
+		}
+	}
+	return TOptional<FVector::FReal>();
+}
+
+TOptional<FVector::FReal> UBoneProxy::GetMultiNumericValue(
+	ESlateTransformComponent::Type Component,
+	ESlateRotationRepresentation::Type Representation,
+	ESlateTransformSubComponent::Type SubComponent,
+	ETransformType TransformType,
+	TArrayView<UBoneProxy*> BoneProxies)
+{
+	TOptional<FVector::FReal> FirstValue = BoneProxies[0]->GetNumericValue(Component, Representation, SubComponent, TransformType);
+	if(!FirstValue.IsSet())
+	{
+		return FirstValue;
+	}
+
+	for(int32 Index=1;Index<BoneProxies.Num();Index++)
+	{
+		TOptional<FVector::FReal> Value = BoneProxies[Index]->GetNumericValue(Component, Representation, SubComponent, TransformType);
+		if(!Value.IsSet())
+		{
+			return Value;
+		}
+		if(!FMath::IsNearlyEqual(FirstValue.GetValue(), Value.GetValue()))
+		{
+			return TOptional<FVector::FReal>();
+		}
+	}
+
+	return FirstValue;
+}
+
+void UBoneProxy::OnNumericValueCommitted(
+	ESlateTransformComponent::Type Component,
+	ESlateRotationRepresentation::Type Representation,
+	ESlateTransformSubComponent::Type SubComponent,
+	FVector::FReal Value, ETextCommit::Type CommitType,
+	ETransformType TransformType,
+	bool bIsCommit)
+{
+	if(TransformType != TransformType_Bone)
+	{
+		return;
+	}
+
+	switch(Component)
+	{
+		case ESlateTransformComponent::Location:
+		{
+			OnPreEditChange(GET_MEMBER_NAME_CHECKED(UBoneProxy, Location), bIsCommit);
+				
+			switch(SubComponent)
+			{
+				case ESlateTransformSubComponent::X:
+				{
+					Location.X = Value;
+					break;
+				}
+				case ESlateTransformSubComponent::Y:
+				{
+					Location.Y = Value;
+					break;
+				}
+				case ESlateTransformSubComponent::Z:
+				{
+					Location.Z = Value;
+					break;
+				}
+			}
+
+			OnPostEditChangeProperty(GET_MEMBER_NAME_CHECKED(UBoneProxy, Location), bIsCommit);
+			break;
+		}
+		case ESlateTransformComponent::Rotation:
+		{
+			OnPreEditChange(GET_MEMBER_NAME_CHECKED(UBoneProxy, Rotation), bIsCommit);
+				
+			switch(SubComponent)
+			{
+				case ESlateTransformSubComponent::Roll:
+				{
+					Rotation.Roll = Value;
+					break;
+				}
+				case ESlateTransformSubComponent::Pitch:
+				{
+					Rotation.Pitch = Value;
+					break;
+				}
+				case ESlateTransformSubComponent::Yaw:
+				{
+					Rotation.Yaw = Value;
+					break;
+				}
+			}
+
+			OnPostEditChangeProperty(GET_MEMBER_NAME_CHECKED(UBoneProxy, Rotation), bIsCommit);
+			break;
+		}
+		case ESlateTransformComponent::Scale:
+		{
+			OnPreEditChange(GET_MEMBER_NAME_CHECKED(UBoneProxy, Scale), bIsCommit);
+
+			switch(SubComponent)
+			{
+				case ESlateTransformSubComponent::X:
+				{
+					Scale.X = Value;
+					break;
+				}
+				case ESlateTransformSubComponent::Y:
+				{
+					Scale.Y = Value;
+					break;
+				}
+				case ESlateTransformSubComponent::Z:
+				{
+					Scale.Z = Value;
+					break;
+				}
+			}
+
+			OnPostEditChangeProperty(GET_MEMBER_NAME_CHECKED(UBoneProxy, Scale), bIsCommit);
+			break;
+		}
+	}
+}
+
+void UBoneProxy::OnMultiNumericValueCommitted(
+	ESlateTransformComponent::Type Component,
+	ESlateRotationRepresentation::Type Representation,
+	ESlateTransformSubComponent::Type SubComponent,
+	FVector::FReal Value,
+	ETextCommit::Type CommitType,
+	ETransformType TransformType,
+	TArrayView<UBoneProxy*> BoneProxies,
+	bool bIsCommit)
+{
+	for(UBoneProxy* BoneProxy : BoneProxies)
+	{
+		BoneProxy->OnNumericValueCommitted(Component, Representation, SubComponent, Value, CommitType, TransformType, bIsCommit);
+	}
+}
+
 void UBoneProxy::PreEditChange(FEditPropertyChain& PropertyAboutToChange)
 {
 	if (UDebugSkelMeshComponent* Component = SkelMeshComponent.Get())
@@ -182,4 +404,20 @@ void UBoneProxy::PostEditChangeProperty(FPropertyChangedEvent& PropertyChangedEv
 	}
 }
 
+void UBoneProxy::OnPreEditChange(FName PropertyName, bool bIsCommit)
+{
+	FEditPropertyChain PropertyChain;
+	FProperty* Property = FindFProperty<FProperty>(GetClass(), PropertyName);
+	PropertyChain.AddHead(Property);
+	PropertyChain.SetActiveMemberPropertyNode(Property);
+	PreEditChange(PropertyChain);
+}
+
+void UBoneProxy::OnPostEditChangeProperty(FName PropertyName, bool bIsCommit)
+{
+	FProperty* Property = FindFProperty<FProperty>(GetClass(), PropertyName);
+	FPropertyChangedEvent ChangedEvent(Property, bIsCommit ? EPropertyChangeType::Unspecified : EPropertyChangeType::Interactive);
+	PostEditChangeProperty(ChangedEvent);
+}
+
 #undef LOCTEXT_NAMESPACE
diff --git a/Engine/Source/Editor/SkeletonEditor/Private/BoneProxyDetailsCustomization.cpp b/Engine/Source/Editor/SkeletonEditor/Private/BoneProxyDetailsCustomization.cpp
--- a/Engine/Source/Editor/SkeletonEditor/Private/BoneProxyDetailsCustomization.cpp
+++ b/Engine/Source/Editor/SkeletonEditor/Private/BoneProxyDetailsCustomization.cpp
@@ -7,13 +7,16 @@
 #include "DetailCategoryBuilder.h"
 #include "DetailLayoutBuilder.h"
 #include "DetailWidgetRow.h"
+#include "IDetailGroup.h"
 #include "IDetailPropertyRow.h"
 #include "Animation/DebugSkelMeshComponent.h"
 #include "ScopedTransaction.h"
 #include "Framework/MultiBox/MultiBoxBuilder.h"
 #include "Widgets/Input/SComboButton.h"
+#include "Widgets/Input/SSegmentedControl.h"
 #include "Widgets/Layout/SBox.h"
 #include "Algo/Transform.h"
+#include "SAdvancedTransformInputBox.h"
 
 #define LOCTEXT_NAMESPACE "FBoneProxyDetailsCustomization"
 
@@ -140,77 +143,223 @@ void FBoneProxyDetailsCustomization::CustomizeDetails(IDetailLayoutBuilder& Deta
 	{
 		bIsEditingEnabled = (Component->AnimScriptInstance == Component->PreviewInstance);
 	}
+	
+	DetailBuilder.HideCategory(TEXT("Transform"));
+	DetailBuilder.HideCategory(TEXT("Reference Transform"));
+	DetailBuilder.HideCategory(TEXT("Mesh Relative Transform"));
+	DetailBuilder.EditCategory(TEXT("Bone")).SetSortOrder(1);
 
-	TSharedRef<IPropertyHandle> LocationProperty = DetailBuilder.GetProperty(GET_MEMBER_NAME_CHECKED(UBoneProxy, Location));
-	TSharedRef<IPropertyHandle> RotationProperty = DetailBuilder.GetProperty(GET_MEMBER_NAME_CHECKED(UBoneProxy, Rotation));
-	TSharedRef<IPropertyHandle> ScaleProperty = DetailBuilder.GetProperty(GET_MEMBER_NAME_CHECKED(UBoneProxy, Scale));
+	IDetailCategoryBuilder& CategoryBuilder = DetailBuilder.EditCategory(TEXT("Transforms"));
+	CategoryBuilder.SetSortOrder(2);
 
-	IDetailCategoryBuilder& CategoryBuilder = DetailBuilder.EditCategory(TEXT("Transform"));
+	const TArray<FText> ButtonLabels =
+	{
+		LOCTEXT("BoneTransform", "Bone"),
+		LOCTEXT("ReferenceTransform", "Reference"),
+		LOCTEXT("MeshTransform", "Mesh Relative")
+	};
+	const TArray<FText> ButtonTooltips =
+	{
+		LOCTEXT("BoneTransformTooltip", "The transform of the bone"),
+		LOCTEXT("ReferenceTransformTooltip", "The reference transform of a bone (original)"),
+		LOCTEXT("MeshTransformTooltip", "The relative transform of the mesh")
+	};
 
-	TSharedPtr<SWidget> NameWidget;
-	TSharedPtr<SWidget> ValueWidget;
+	TSharedPtr<SSegmentedControl<UBoneProxy::ETransformType>> TransformChoiceWidget;
 
-	IDetailPropertyRow& LocationPropertyRow = CategoryBuilder.AddProperty(LocationProperty);
-	LocationPropertyRow.GetDefaultWidgets(NameWidget, ValueWidget);
-	LocationPropertyRow.OverrideResetToDefault(FResetToDefaultOverride::Create(FIsResetToDefaultVisible::CreateSP(this, &FBoneProxyDetailsCustomization::IsResetLocationVisible, BoneProxiesView), FResetToDefaultHandler::CreateSP(this, &FBoneProxyDetailsCustomization::HandleResetLocation, BoneProxiesView)));
+	// use a static shared ref so that all views retain these settings
+	static TSharedRef<TArray<UBoneProxy::ETransformType>> VisibleTransforms =
+		MakeShareable(new TArray<UBoneProxy::ETransformType>({
+			UBoneProxy::TransformType_Bone,
+			UBoneProxy::TransformType_Reference,
+			UBoneProxy::TransformType_Mesh}));
 
-	LocationPropertyRow.CustomWidget()
-	.NameContent()
-	[
-		BuildTransformFieldLabel(&FirstBoneProxy->bLocalLocation, LOCTEXT("Location", "Location"), Objects.Num() > 1)
-	]
+	CategoryBuilder.AddCustomRow(FText::FromString(TEXT("TransformType")))
 	.ValueContent()
-	.MinDesiredWidth(BoneProxyCustomizationConstants::ItemWidth * 3.0f)
-	.MaxDesiredWidth(BoneProxyCustomizationConstants::ItemWidth * 3.0f)
+	.MinDesiredWidth(375.f)
+	.MaxDesiredWidth(375.f)
+	.HAlign(HAlign_Left)
 	[
-		SNew(SBox)
-		.IsEnabled(bIsEditingEnabled)
+		SNew(SHorizontalBox)
+		+SHorizontalBox::Slot()
+		.AutoWidth()
+		.HAlign(HAlign_Left)
+		.VAlign(VAlign_Center)
 		[
-			ValueWidget.ToSharedRef()
+			SAssignNew(TransformChoiceWidget, SSegmentedControl<UBoneProxy::ETransformType>)
+			.SupportsMultiSelection(true)
+			.Values_Lambda([]()
+			{
+				return VisibleTransforms.Get();
+			})
+			.OnValuesChanged_Lambda([](TArray<UBoneProxy::ETransformType> Values)
+			{
+				VisibleTransforms.Get() = Values;
+			})
+			+ SSegmentedControl<UBoneProxy::ETransformType>::Slot(UBoneProxy::TransformType_Bone)
+			.Text(ButtonLabels[0])
+			.ToolTip(ButtonTooltips[0])
+			+ SSegmentedControl<UBoneProxy::ETransformType>::Slot(UBoneProxy::TransformType_Reference)
+			.Text(ButtonLabels[1])
+			.ToolTip(ButtonTooltips[1])
+			+ SSegmentedControl<UBoneProxy::ETransformType>::Slot(UBoneProxy::TransformType_Mesh)
+			.Text(ButtonLabels[2])
+			.ToolTip(ButtonTooltips[2])
 		]
 	];
 
-	IDetailPropertyRow& RotationPropertyRow = CategoryBuilder.AddProperty(RotationProperty);
-	RotationPropertyRow.GetDefaultWidgets(NameWidget, ValueWidget);
-	RotationPropertyRow.OverrideResetToDefault(FResetToDefaultOverride::Create(FIsResetToDefaultVisible::CreateSP(this, &FBoneProxyDetailsCustomization::IsResetRotationVisible, BoneProxiesView), FResetToDefaultHandler::CreateSP(this, &FBoneProxyDetailsCustomization::HandleResetRotation, BoneProxiesView)));
+	SAdvancedTransformInputBox<FEulerTransform>::FArguments TransformWidgetArgs = SAdvancedTransformInputBox<FEulerTransform>::FArguments()
+	.DisplayRelativeWorld(true)
+	.AllowEditRotationRepresentation(false)
+	.DisplayScaleLock(true)
+	.Font(IDetailLayoutBuilder::GetDetailFont())
+	.UseQuaternionForRotation(false)
+	.OnGetIsComponentRelative_Lambda(
+		[BoneProxiesView](ESlateTransformComponent::Type InComponent)
+		{
+			switch(InComponent)
+			{
+				case ESlateTransformComponent::Location:
+					return BoneProxiesView[0]->bLocalLocation;
+				case ESlateTransformComponent::Rotation:
+					return BoneProxiesView[0]->bLocalRotation;
+				case ESlateTransformComponent::Scale:
+					return BoneProxiesView[0]->bLocalScale;
+			}
+			return true;
+		})
+	.OnIsComponentRelativeChanged_Lambda(
+		[BoneProxiesView](ESlateTransformComponent::Type InComponent, bool bIsRelative)
+		{
+			for(UBoneProxy* BoneProxy : BoneProxiesView)
+			{
+				switch(InComponent)
+				{
+					case ESlateTransformComponent::Location:
+					{
+						BoneProxy->bLocalLocation = bIsRelative;
+						break;
+					}
+					case ESlateTransformComponent::Rotation:
+					{
+						BoneProxy->bLocalRotation = bIsRelative;
+						break;
+					}
+					case ESlateTransformComponent::Scale:
+					{
+						BoneProxy->bLocalScale = bIsRelative;
+						break;
+					}
+				}
+			}
+		});
+
+	TArray<TSharedRef<IPropertyHandle>> Properties;
+	Properties.Add(DetailBuilder.GetProperty(GET_MEMBER_NAME_CHECKED(UBoneProxy, Location)));
+	Properties.Add(DetailBuilder.GetProperty(GET_MEMBER_NAME_CHECKED(UBoneProxy, Rotation)));
+	Properties.Add(DetailBuilder.GetProperty(GET_MEMBER_NAME_CHECKED(UBoneProxy, Scale)));
+	Properties.Add(DetailBuilder.GetProperty(GET_MEMBER_NAME_CHECKED(UBoneProxy, ReferenceLocation)));
+	Properties.Add(DetailBuilder.GetProperty(GET_MEMBER_NAME_CHECKED(UBoneProxy, ReferenceRotation)));
+	Properties.Add(DetailBuilder.GetProperty(GET_MEMBER_NAME_CHECKED(UBoneProxy, ReferenceScale)));
+	Properties.Add(DetailBuilder.GetProperty(GET_MEMBER_NAME_CHECKED(UBoneProxy, MeshLocation)));
+	Properties.Add(DetailBuilder.GetProperty(GET_MEMBER_NAME_CHECKED(UBoneProxy, MeshRotation)));
+	Properties.Add(DetailBuilder.GetProperty(GET_MEMBER_NAME_CHECKED(UBoneProxy, MeshScale)));
+
+	int32 PropertyIndex = 0;
+	for(int32 TransformIndex=0;TransformIndex<3;TransformIndex++)
+	{
+		// only the first transform can be edited
+		const UBoneProxy::ETransformType TransformType = (UBoneProxy::ETransformType)TransformIndex; 
+		bIsEditingEnabled = TransformType == UBoneProxy::TransformType_Bone ? bIsEditingEnabled : false;
 
-	RotationPropertyRow.CustomWidget()
-	.NameContent()
-	[
-		BuildTransformFieldLabel(&FirstBoneProxy->bLocalRotation, LOCTEXT("Rotation", "Rotation"), Objects.Num() > 1)
-	]
-	.ValueContent()
-	.MinDesiredWidth(BoneProxyCustomizationConstants::ItemWidth * 3.0f)
-	.MaxDesiredWidth(BoneProxyCustomizationConstants::ItemWidth * 3.0f)
-	[
-		SNew(SBox)
+		TransformWidgetArgs
 		.IsEnabled(bIsEditingEnabled)
+		.DisplayRelativeWorld(bIsEditingEnabled)
+		.DisplayScaleLock(bIsEditingEnabled)
+		.OnGetNumericValue_Static(&UBoneProxy::GetMultiNumericValue, TransformType, BoneProxiesView);
+
+		if(bIsEditingEnabled)
+		{
+			TransformWidgetArgs.OnNumericValueChanged_Static(&UBoneProxy::OnMultiNumericValueCommitted, ETextCommit::Default, TransformType, BoneProxiesView, false);
+			TransformWidgetArgs.OnNumericValueCommitted_Static (&UBoneProxy::OnMultiNumericValueCommitted, TransformType, BoneProxiesView, true);
+		}
+		else
+		{
+			TransformWidgetArgs._OnNumericValueChanged.Unbind();
+			TransformWidgetArgs._OnNumericValueCommitted.Unbind();
+		}
+		
+		IDetailGroup& Group = CategoryBuilder.AddGroup(*ButtonLabels[TransformIndex].ToString(), ButtonLabels[TransformIndex], false, true);
+		Group.HeaderRow()
+		.Visibility( TAttribute<EVisibility>::CreateLambda([TransformChoiceWidget, TransformType]() -> EVisibility
+		{
+			return TransformChoiceWidget->HasValue(TransformType) ? EVisibility::Visible : EVisibility::Collapsed;
+		}))
+		.NameContent()
+		[
+			SNew(STextBlock)
+			.Font(IDetailLayoutBuilder::GetDetailFont())
+			.Text(ButtonLabels[TransformIndex])
+			.ToolTipText(ButtonTooltips[TransformIndex])
+		];
+
+		IDetailPropertyRow& LocationPropertyRow = Group.AddPropertyRow(Properties[PropertyIndex++]);
+		if(bIsEditingEnabled)
+		{
+			LocationPropertyRow.OverrideResetToDefault(FResetToDefaultOverride::Create(FIsResetToDefaultVisible::CreateSP(this, &FBoneProxyDetailsCustomization::IsResetLocationVisible, BoneProxiesView), FResetToDefaultHandler::CreateSP(this, &FBoneProxyDetailsCustomization::HandleResetLocation, BoneProxiesView)));
+		}
+
+		LocationPropertyRow.CustomWidget()
+		.NameContent()
+		.HAlign(HAlign_Fill)
 		[
-			ValueWidget.ToSharedRef()
+			SAdvancedTransformInputBox<FEulerTransform>::ConstructLabel(TransformWidgetArgs, ESlateTransformComponent::Location)
 		]
-	];
+		.ValueContent()
+		.MinDesiredWidth(BoneProxyCustomizationConstants::ItemWidth * 3.0f)
+		.MaxDesiredWidth(BoneProxyCustomizationConstants::ItemWidth * 3.0f)
+		[
+			SAdvancedTransformInputBox<FEulerTransform>::ConstructWidget(TransformWidgetArgs, ESlateTransformComponent::Location)
+		];
 
-	IDetailPropertyRow& ScalePropertyRow = CategoryBuilder.AddProperty(ScaleProperty);
-	ScalePropertyRow.GetDefaultWidgets(NameWidget, ValueWidget);
-	ScalePropertyRow.OverrideResetToDefault(FResetToDefaultOverride::Create(FIsResetToDefaultVisible::CreateSP(this, &FBoneProxyDetailsCustomization::IsResetScaleVisible, BoneProxiesView), FResetToDefaultHandler::CreateSP(this, &FBoneProxyDetailsCustomization::HandleResetScale, BoneProxiesView)));
+		IDetailPropertyRow& RotationPropertyRow = Group.AddPropertyRow(Properties[PropertyIndex++]);
+		if(bIsEditingEnabled)
+		{
+			RotationPropertyRow.OverrideResetToDefault(FResetToDefaultOverride::Create(FIsResetToDefaultVisible::CreateSP(this, &FBoneProxyDetailsCustomization::IsResetRotationVisible, BoneProxiesView), FResetToDefaultHandler::CreateSP(this, &FBoneProxyDetailsCustomization::HandleResetRotation, BoneProxiesView)));
+		}
 
-	ScalePropertyRow.CustomWidget()
-	.NameContent()
-	[
-		SNew(STextBlock)
-		.Text(ScaleProperty->GetPropertyDisplayName())
-		.Font(DetailBuilder.GetDetailFont())
-	]
-	.ValueContent()
-	.MinDesiredWidth(BoneProxyCustomizationConstants::ItemWidth * 3.0f)
-	.MaxDesiredWidth(BoneProxyCustomizationConstants::ItemWidth * 3.0f)
-	[
-		SNew(SBox)
-		.IsEnabled(bIsEditingEnabled)
+		RotationPropertyRow.CustomWidget()
+		.NameContent()
+		.HAlign(HAlign_Fill)
 		[
-			ValueWidget.ToSharedRef()
+			SAdvancedTransformInputBox<FEulerTransform>::ConstructLabel(TransformWidgetArgs, ESlateTransformComponent::Rotation)
 		]
-	];
+		.ValueContent()
+		.MinDesiredWidth(BoneProxyCustomizationConstants::ItemWidth * 3.0f)
+		.MaxDesiredWidth(BoneProxyCustomizationConstants::ItemWidth * 3.0f)
+		[
+			SAdvancedTransformInputBox<FEulerTransform>::ConstructWidget(TransformWidgetArgs, ESlateTransformComponent::Rotation)
+		];
+
+		IDetailPropertyRow& ScalePropertyRow = Group.AddPropertyRow(Properties[PropertyIndex++]);
+		if(bIsEditingEnabled)
+		{
+			ScalePropertyRow.OverrideResetToDefault(FResetToDefaultOverride::Create(FIsResetToDefaultVisible::CreateSP(this, &FBoneProxyDetailsCustomization::IsResetScaleVisible, BoneProxiesView), FResetToDefaultHandler::CreateSP(this, &FBoneProxyDetailsCustomization::HandleResetScale, BoneProxiesView)));
+		}
+
+		ScalePropertyRow.CustomWidget()
+		.NameContent()
+		.HAlign(HAlign_Fill)
+		[
+			SAdvancedTransformInputBox<FEulerTransform>::ConstructLabel(TransformWidgetArgs, ESlateTransformComponent::Scale)
+		]
+		.ValueContent()
+		.MinDesiredWidth(BoneProxyCustomizationConstants::ItemWidth * 3.0f)
+		.MaxDesiredWidth(BoneProxyCustomizationConstants::ItemWidth * 3.0f)
+		[
+			SAdvancedTransformInputBox<FEulerTransform>::ConstructWidget(TransformWidgetArgs, ESlateTransformComponent::Scale)
+		];
+	}
 }
 
 bool FBoneProxyDetailsCustomization::IsResetLocationVisible(TSharedPtr<IPropertyHandle> InPropertyHandle, TArrayView<UBoneProxy*> InBoneProxies)
diff --git a/Engine/Source/Editor/SkeletonEditor/Public/BoneProxy.h b/Engine/Source/Editor/SkeletonEditor/Public/BoneProxy.h
--- a/Engine/Source/Editor/SkeletonEditor/Public/BoneProxy.h
+++ b/Engine/Source/Editor/SkeletonEditor/Public/BoneProxy.h
@@ -7,6 +7,7 @@
 #include "TickableEditorObject.h"
 #include "UObject/Object.h"
 #include "UObject/WeakObjectPtr.h"
+#include "SAdvancedTransformInputBox.h"
 #include "BoneProxy.generated.h"
 
 
@@ -21,10 +22,21 @@ class SKELETONEDITOR_API UBoneProxy : public UObject, public FTickableEditorObje
 public:
 	UBoneProxy();
 
+	enum ETransformType
+	{
+		TransformType_Bone,
+		TransformType_Reference,
+		TransformType_Mesh
+	};
+
 	/** UObject interface */
 	virtual void PreEditChange(FEditPropertyChain& PropertyAboutToChange) override;
 	virtual void PostEditChangeProperty(FPropertyChangedEvent& PropertyChangedEvent) override;
 
+	// overload based on a property name
+	void OnPreEditChange(FName PropertyName, bool bIsCommit);
+	void OnPostEditChangeProperty(FName PropertyName, bool bIsCommit);
+	
 	/** FTickableEditorObject interface */
 	virtual void Tick(float DeltaTime) override;
 	virtual bool IsTickable() const override;
@@ -80,6 +92,9 @@ public:
 	/** Whether to use local or world rotation */
 	bool bLocalRotation;
 
+	/** Whether to use local or world rotation */
+	bool bLocalScale;
+
 	// Handle property deltas
 	FVector PreviousLocation;
 	FRotator PreviousRotation;
@@ -90,4 +105,40 @@ public:
 
 	/** Flag indicating whether this FTickableEditorObject should actually tick */
 	bool bIsTickable;
-};
\ No newline at end of file
+
+	/** Method to react to retrieval of numeric values for the widget */
+	TOptional<FVector::FReal> GetNumericValue(
+		ESlateTransformComponent::Type Component,
+		ESlateRotationRepresentation::Type Representation,
+		ESlateTransformSubComponent::Type SubComponent,
+		ETransformType TransformType) const;
+
+	/** Static method to retrieve a value for a list of proxies */
+	static TOptional<FVector::FReal> GetMultiNumericValue(
+		ESlateTransformComponent::Type Component,
+		ESlateRotationRepresentation::Type Representation,
+		ESlateTransformSubComponent::Type SubComponent,
+		ETransformType TransformType,
+		TArrayView<UBoneProxy*> BoneProxies);
+
+	/** Method to react to changes of numeric values in the widget */
+	void OnNumericValueCommitted(
+		ESlateTransformComponent::Type Component,
+		ESlateRotationRepresentation::Type Representation,
+		ESlateTransformSubComponent::Type SubComponent,
+		FVector::FReal Value,
+		ETextCommit::Type CommitType,
+		ETransformType TransformType,
+		bool bIsCommit);
+
+	/** Method to react to changes of numeric values in the widget */
+	static void OnMultiNumericValueCommitted(
+		ESlateTransformComponent::Type Component,
+		ESlateRotationRepresentation::Type Representation,
+		ESlateTransformSubComponent::Type SubComponent,
+		FVector::FReal Value,
+		ETextCommit::Type CommitType,
+		ETransformType TransformType,
+		TArrayView<UBoneProxy*> BoneProxies,
+		bool bIsCommit);
+};
diff --git a/Engine/Source/Editor/SkeletonEditor/SkeletonEditor.Build.cs b/Engine/Source/Editor/SkeletonEditor/SkeletonEditor.Build.cs
--- a/Engine/Source/Editor/SkeletonEditor/SkeletonEditor.Build.cs
+++ b/Engine/Source/Editor/SkeletonEditor/SkeletonEditor.Build.cs
@@ -29,12 +29,15 @@ public class SkeletonEditor : ModuleRules
 				"ToolMenus",
 				"EditorWidgets",
 				"ToolWidgets",
+				"AnimationCore",
+				"AnimationWidgets",
 			}
 		);
 
         PublicIncludePathModuleNames.AddRange(
             new string[] {
                 "Persona",
+                "AnimationWidgets",
             }
         );
 
