Added virtual bone support to blend profiles

diff --git a/Engine/Source/Editor/SkeletonEditor/Private/SkeletonTreeVirtualBoneItem.cpp b/Engine/Source/Editor/SkeletonEditor/Private/SkeletonTreeVirtualBoneItem.cpp
--- a/Engine/Source/Editor/SkeletonEditor/Private/SkeletonTreeVirtualBoneItem.cpp
+++ b/Engine/Source/Editor/SkeletonEditor/Private/SkeletonTreeVirtualBoneItem.cpp
@@ -9,6 +9,7 @@
 #include "Widgets/Text/SInlineEditableTextBlock.h"
 #include "Containers/UnrealString.h"
 #include "Animation/DebugSkelMeshComponent.h"
+#include "Animation/BlendProfile.h"
 #include "UObject/Package.h"
 
 #define LOCTEXT_NAMESPACE "FSkeletonTreeVirtualBoneItem"
@@ -88,9 +89,77 @@ void FSkeletonTreeVirtualBoneItem::GenerateWidgetForNameColumn(TSharedPtr< SHori
 
 TSharedRef< SWidget > FSkeletonTreeVirtualBoneItem::GenerateWidgetForDataColumn(const FName& DataColumnName, FIsSelected InIsSelected)
 {
+	if (DataColumnName == ISkeletonTree::Columns::BlendProfile)
+	{
+		return SNew(SBox)
+			.Padding(0.0f)
+			.HAlign(HAlign_Fill)
+			.VAlign(VAlign_Center)
+			[
+				SNew(SSpinBox<float>)
+				.Visibility(this, &FSkeletonTreeVirtualBoneItem::GetBoneBlendProfileVisibility)
+			.Style(&FAppStyle::Get(), "SkeletonTree.HyperlinkSpinBox")
+			.Font(FAppStyle::Get().GetFontStyle("SmallFont"))
+			.ContentPadding(0.0f)
+			.Delta(0.01f)
+			.MinValue(0.0f)
+			.MinSliderValue(this, &FSkeletonTreeVirtualBoneItem::GetBlendProfileMinSliderValue)
+			.MaxSliderValue(this, &FSkeletonTreeVirtualBoneItem::GetBlendProfileMaxSliderValue)
+			.Value(this, &FSkeletonTreeVirtualBoneItem::GetBoneBlendProfileScale)
+			.OnValueCommitted(this, &FSkeletonTreeVirtualBoneItem::OnBlendSliderCommitted)
+			.OnValueChanged(this, &FSkeletonTreeVirtualBoneItem::OnBlendSliderCommitted, ETextCommit::OnEnter)
+			];
+	}
+
 	return SNullWidget::NullWidget;
 }
 
+EVisibility FSkeletonTreeVirtualBoneItem::GetBoneBlendProfileVisibility() const
+{
+	return GetSkeletonTree()->GetSelectedBlendProfile() ? EVisibility::Visible : EVisibility::Collapsed;
+}
+
+float FSkeletonTreeVirtualBoneItem::GetBoneBlendProfileScale()	const
+{
+	if (UBlendProfile* CurrentProfile = GetSkeletonTree()->GetSelectedBlendProfile())
+	{
+		return CurrentProfile->GetBoneBlendScale(BoneName);
+	}
+
+	return 0.0;
+}
+
+TOptional<float> FSkeletonTreeVirtualBoneItem::GetBlendProfileMaxSliderValue() const
+{
+	if (UBlendProfile* CurrentProfile = GetSkeletonTree()->GetSelectedBlendProfile())
+	{
+		return (CurrentProfile->GetMode() == EBlendProfileMode::WeightFactor) ? 10.0f : 1.0f;
+	}
+
+	return 1.0f;
+}
+
+TOptional<float> FSkeletonTreeVirtualBoneItem::GetBlendProfileMinSliderValue() const
+{
+	if (UBlendProfile* CurrentProfile = GetSkeletonTree()->GetSelectedBlendProfile())
+	{
+		return (CurrentProfile->GetMode() == EBlendProfileMode::WeightFactor) ? 1.0f : 0.0f;
+	}
+
+	return 0.0f;
+}
+
+void FSkeletonTreeVirtualBoneItem::OnBlendSliderCommitted(float NewValue, ETextCommit::Type CommitType)
+{
+	SetBoneBlendProfileScale(NewValue, false);
+}
+
+void FSkeletonTreeVirtualBoneItem::SetBoneBlendProfileScale(float NewScale, bool bRecurse)
+{
+	FName BlendProfileName = GetSkeletonTree()->GetSelectedBlendProfile()->GetFName();
+	GetEditableSkeleton()->SetBlendProfileScale(BlendProfileName, BoneName, NewScale, bRecurse);
+}
+
 FSlateFontInfo FSkeletonTreeVirtualBoneItem::GetBoneTextFont() const
 {
 	return FEditorStyle::GetWidgetStyle<FTextBlockStyle>("SkeletonTree.NormalFont").Font;
diff --git a/Engine/Source/Editor/SkeletonEditor/Private/SkeletonTreeVirtualBoneItem.h b/Engine/Source/Editor/SkeletonEditor/Private/SkeletonTreeVirtualBoneItem.h
--- a/Engine/Source/Editor/SkeletonEditor/Private/SkeletonTreeVirtualBoneItem.h
+++ b/Engine/Source/Editor/SkeletonEditor/Private/SkeletonTreeVirtualBoneItem.h
@@ -47,6 +47,28 @@ public:
 	void EnableBoneProxyTick(bool bEnable);
 
 private:
+	
+	/** Callback from a slider widget if the text entry or slider movement is used */
+	void OnBlendSliderCommitted(float NewValue, ETextCommit::Type CommitType);
+
+	/** Set Translation Retargeting Mode for this bone. */
+	void SetBoneTranslationRetargetingMode(EBoneTranslationRetargetingMode::Type NewRetargetingMode);
+
+	/** Set current Blend Scale for this bone */
+	void SetBoneBlendProfileScale(float NewScale, bool bRecurse);
+
+	/** Get the current Blend Scale for this bone */
+	float GetBoneBlendProfileScale() const;
+
+	/** Get the max slider value supported by the selected blend profile's mode*/
+	TOptional<float> GetBlendProfileMaxSliderValue() const;
+
+	/** Get the min slider value supported by the selected blend profile's mode*/
+	TOptional<float> GetBlendProfileMinSliderValue() const;
+
+	/** Used to hide the Bone Profile row widget if there is no current Blend Profile */
+	EVisibility GetBoneBlendProfileVisibility() const;
+	
 	/** Called when we are about to rename a virtual bone */
 	void OnVirtualBoneNameEditing();
 
diff --git a/Engine/Source/Runtime/Engine/Classes/Animation/BlendProfile.h b/Engine/Source/Runtime/Engine/Classes/Animation/BlendProfile.h
--- a/Engine/Source/Runtime/Engine/Classes/Animation/BlendProfile.h
+++ b/Engine/Source/Runtime/Engine/Classes/Animation/BlendProfile.h
@@ -82,6 +82,11 @@ public:
 	 */
 	void SetBoneBlendScale(const FName& InBoneName, float InScale, bool bRecurse = false, bool bCreate = false);
 
+	/** Removes the entry for the specified bone index (does nothing if it doesn't exist) 
+	*  @param InBoneIdx Index of the bone to remove from this blend profile
+	*/
+	void RemoveEntry(int32 InBoneIdx);
+
 	/** Get the set blend scale for the specified bone, will return 1.0f if no entry was found (no scale)
 	 *  @param InBoneIdx Index of the bone to retrieve
 	 */
diff --git a/Engine/Source/Runtime/Engine/Private/Animation/BlendProfile.cpp b/Engine/Source/Runtime/Engine/Private/Animation/BlendProfile.cpp
--- a/Engine/Source/Runtime/Engine/Private/Animation/BlendProfile.cpp
+++ b/Engine/Source/Runtime/Engine/Private/Animation/BlendProfile.cpp
@@ -40,6 +40,15 @@ void UBlendProfile::SetBoneBlendScale(const FName& InBoneName, float InScale, bo
 	SetBoneBlendScale(BoneIndex, InScale, bRecurse, bCreate);
 }
 
+void UBlendProfile::RemoveEntry(int32 InBoneIdx)
+{
+	Modify();
+	ProfileEntries.RemoveAll([InBoneIdx](const FBlendProfileBoneEntry& Current)
+		{
+			return Current.BoneReference.BoneIndex == InBoneIdx;
+		});
+}
+
 float UBlendProfile::GetBoneBlendScale(int32 InBoneIdx) const
 {
 	const FBlendProfileBoneEntry* FoundEntry = ProfileEntries.FindByPredicate([InBoneIdx](const FBlendProfileBoneEntry& Entry)
diff --git a/Engine/Source/Runtime/Engine/Private/Animation/Skeleton.cpp b/Engine/Source/Runtime/Engine/Private/Animation/Skeleton.cpp
--- a/Engine/Source/Runtime/Engine/Private/Animation/Skeleton.cpp
+++ b/Engine/Source/Runtime/Engine/Private/Animation/Skeleton.cpp
@@ -2231,6 +2231,16 @@ void USkeleton::RemoveVirtualBones(const TArray<FName>& BonesToRemove)
 				}
 			}
 			VirtualBones.RemoveAt(Idx,1,false);
+
+			// @todo: This might be a slow operation if there's a large amount of blend profiles and entries
+			int32 BoneIdx = GetReferenceSkeleton().FindBoneIndex(BoneName);
+			if(BoneIdx != INDEX_NONE)
+			{
+				for (UBlendProfile* Profile : BlendProfiles)
+				{
+					Profile->RemoveEntry(BoneIdx);
+				}
+			}
 		}
 	}
 
