Applying Blend Profiles with spaces in name to Blend Settings in AnimBPs does not work
#fix ensure that users cannot generate blend profile names with invalid characters (as it is used to generate a UObject with)
#misc Added OnVerifyTextChanged to STextEntryPopup which routes it to its inner SEditableTextBox

diff --git a/Engine/Source/Editor/SkeletonEditor/Private/SBlendProfilePicker.cpp b/Engine/Source/Editor/SkeletonEditor/Private/SBlendProfilePicker.cpp
--- a/Engine/Source/Editor/SkeletonEditor/Private/SBlendProfilePicker.cpp
+++ b/Engine/Source/Editor/SkeletonEditor/Private/SBlendProfilePicker.cpp
@@ -463,7 +463,11 @@ void SBlendProfilePicker::OnCreateNewProfile(EBlendProfileMode InMode)
 {
 	TSharedRef<STextEntryPopup> TextEntry = SNew(STextEntryPopup)
 		.Label(LOCTEXT("NewProfileName", "Profile Name"))
-		.OnTextCommitted(this, &SBlendProfilePicker::OnCreateNewProfileComitted, InMode);
+		.OnTextCommitted(this, &SBlendProfilePicker::OnCreateNewProfileComitted, InMode)
+		.OnVerifyTextChanged_Lambda([](const FText& InNewText, FText& OutErrorMessage) -> bool
+		{
+			return FName::IsValidXName(InNewText.ToString(), INVALID_OBJECTNAME_CHARACTERS INVALID_LONGPACKAGE_CHARACTERS, &OutErrorMessage);
+		});
 
 	FSlateApplication::Get().PushMenu(
 		AsShared(),
diff --git a/Engine/Source/Editor/SkeletonEditor/Private/SSkeletonTree.cpp b/Engine/Source/Editor/SkeletonEditor/Private/SSkeletonTree.cpp
--- a/Engine/Source/Editor/SkeletonEditor/Private/SSkeletonTree.cpp
+++ b/Engine/Source/Editor/SkeletonEditor/Private/SSkeletonTree.cpp
@@ -601,12 +601,16 @@ void SSkeletonTree::CreateTreeColumns()
 			.Text_Lambda([this] () -> FText
 			{
 				FName CurrentProfile = BlendProfilePicker->GetSelectedBlendProfileName();
-				return (CurrentProfile != NAME_None) ? FText::FromName(CurrentProfile) : LOCTEXT("NoBlendProfile", "No Blend");
+				return (CurrentProfile != NAME_None) ? FText::FromName(CurrentProfile) : LOCTEXT("NoBlendProfile", "NoBlend");
 			})
 			.OnTextCommitted_Lambda([this](const FText& InText, ETextCommit::Type InCommitType)
 			{
 				BlendProfilePicker->OnCreateNewProfileComitted(InText, InCommitType, NewBlendProfileMode);
 			})
+			.OnVerifyTextChanged_Lambda([](const FText& InNewText, FText& OutErrorMessage) -> bool
+			{
+				return FName::IsValidXName(InNewText.ToString(), INVALID_OBJECTNAME_CHARACTERS INVALID_LONGPACKAGE_CHARACTERS, &OutErrorMessage);
+			})
 			.IsReadOnly(true)
 		]
 	];
diff --git a/Engine/Source/Runtime/Slate/Private/Widgets/Input/STextEntryPopup.cpp b/Engine/Source/Runtime/Slate/Private/Widgets/Input/STextEntryPopup.cpp
--- a/Engine/Source/Runtime/Slate/Private/Widgets/Input/STextEntryPopup.cpp
+++ b/Engine/Source/Runtime/Slate/Private/Widgets/Input/STextEntryPopup.cpp
@@ -38,6 +38,7 @@ void STextEntryPopup::Construct( const FArguments& InArgs )
 					.Text( InArgs._DefaultText )
 					.OnTextCommitted( InArgs._OnTextCommitted )
 					.OnTextChanged( InArgs._OnTextChanged )
+					.OnVerifyTextChanged( InArgs._OnVerifyTextChanged )
 					.HintText( InArgs._HintText )
 					.SelectAllTextWhenFocused( InArgs._SelectAllTextWhenFocused )
 					.ClearKeyboardFocusOnCommit( InArgs._ClearKeyboardFocusOnCommit )
diff --git a/Engine/Source/Runtime/Slate/Public/Widgets/Input/STextEntryPopup.h b/Engine/Source/Runtime/Slate/Public/Widgets/Input/STextEntryPopup.h
--- a/Engine/Source/Runtime/Slate/Public/Widgets/Input/STextEntryPopup.h
+++ b/Engine/Source/Runtime/Slate/Public/Widgets/Input/STextEntryPopup.h
@@ -44,6 +44,9 @@ public:
 		/** The maximum width for text entry */
 		SLATE_ATTRIBUTE( float, MaxWidth )
 
+		/** Called whenever the text is changed programmatically or interactively by the user */
+		SLATE_EVENT( FOnVerifyTextChanged, OnVerifyTextChanged )
+
 		/** Provide a alternative mechanism for error reporting. */
 		SLATE_ARGUMENT( TSharedPtr<class IErrorReportingWidget>, ErrorReporting )
 
