Skeleton Editor: Transform Widget Rework

diff --git a/Engine/Source/Editor/SkeletonEditor/Private/BoneProxy.cpp b/Engine/Source/Editor/SkeletonEditor/Private/BoneProxy.cpp
--- a/Engine/Source/Editor/SkeletonEditor/Private/BoneProxy.cpp
+++ b/Engine/Source/Editor/SkeletonEditor/Private/BoneProxy.cpp
@@ -268,6 +268,89 @@ void UBoneProxy::OnMultiNumericValueCommitted(
 	}
 }
 
+bool UBoneProxy::DiffersFromDefault(ESlateTransformComponent::Type Component, ETransformType TransformType) const
+{
+	if(TransformType == TransformType_Bone)
+	{
+		switch(Component)
+		{
+			case ESlateTransformComponent::Location:
+			{
+				return !Location.Equals(ReferenceLocation);
+			}
+			case ESlateTransformComponent::Rotation:
+			{
+				return !Rotation.Equals(ReferenceRotation);
+			}
+			case ESlateTransformComponent::Scale:
+			{
+				return !Scale.Equals(ReferenceScale);
+			}
+			default:
+			{
+				return DiffersFromDefault(ESlateTransformComponent::Location, TransformType) ||
+					DiffersFromDefault(ESlateTransformComponent::Rotation, TransformType) ||
+						DiffersFromDefault(ESlateTransformComponent::Scale, TransformType);
+			}
+		}
+	}
+	return false;
+}
+
+void UBoneProxy::ResetToDefault(ESlateTransformComponent::Type InComponent, ETransformType TransformType)
+{
+	if(TransformType == TransformType_Bone)
+	{
+		if (UDebugSkelMeshComponent* Component = SkelMeshComponent.Get())
+		{
+			if (Component->PreviewInstance && Component->AnimScriptInstance == Component->PreviewInstance)
+			{
+				int32 BoneIndex = Component->GetBoneIndex(BoneName);
+				if (BoneIndex != INDEX_NONE && BoneIndex < Component->GetNumComponentSpaceTransforms())
+				{
+					Component->PreviewInstance->SetFlags(RF_Transactional);
+					Component->PreviewInstance->Modify();
+
+					FAnimNode_ModifyBone& ModifyBone = Component->PreviewInstance->ModifyBone(BoneName);
+
+					switch(InComponent)
+					{
+						case ESlateTransformComponent::Location:
+						{
+							ModifyBone.Translation = ReferenceLocation;
+							break;
+						}
+						case ESlateTransformComponent::Rotation:
+						{
+							ModifyBone.Rotation = ReferenceRotation;
+							break;
+						}
+						case ESlateTransformComponent::Scale:
+						{
+							ModifyBone.Scale = ReferenceScale;
+							break;
+						}
+						default:
+						{
+							ModifyBone.Translation = ReferenceLocation;
+							ModifyBone.Rotation = ReferenceRotation;
+							ModifyBone.Scale = ReferenceScale;
+							break;
+						}
+					}
+
+					if(ModifyBone.Translation.Equals(ReferenceLocation) &&
+						ModifyBone.Rotation.Equals(ReferenceRotation) &&
+						ModifyBone.Scale.Equals(ReferenceScale))
+					{
+						Component->PreviewInstance->RemoveBoneModification(BoneName);
+					}
+				}
+			}
+		}
+	}
+}
+
 void UBoneProxy::PreEditChange(FEditPropertyChain& PropertyAboutToChange)
 {
 	if (UDebugSkelMeshComponent* Component = SkelMeshComponent.Get())
diff --git a/Engine/Source/Editor/SkeletonEditor/Private/BoneProxyDetailsCustomization.cpp b/Engine/Source/Editor/SkeletonEditor/Private/BoneProxyDetailsCustomization.cpp
--- a/Engine/Source/Editor/SkeletonEditor/Private/BoneProxyDetailsCustomization.cpp
+++ b/Engine/Source/Editor/SkeletonEditor/Private/BoneProxyDetailsCustomization.cpp
@@ -18,6 +18,12 @@
 #include "Algo/Transform.h"
 #include "SAdvancedTransformInputBox.h"
 
+#if WITH_EDITOR
+
+#include "HAL/PlatformApplicationMisc.h"
+
+#endif
+
 #define LOCTEXT_NAMESPACE "FBoneProxyDetailsCustomization"
 
 namespace BoneProxyCustomizationConstants
@@ -165,7 +171,21 @@ void FBoneProxyDetailsCustomization::CustomizeDetails(IDetailLayoutBuilder& Deta
 		LOCTEXT("MeshTransformTooltip", "The relative transform of the mesh")
 	};
 
-	TSharedPtr<SSegmentedControl<UBoneProxy::ETransformType>> TransformChoiceWidget;
+	const TArray<UBoneProxy::ETransformType> TransformTypes = {
+		UBoneProxy::TransformType_Bone,
+		UBoneProxy::TransformType_Reference,
+		UBoneProxy::TransformType_Mesh
+	};
+
+	static TAttribute<TArray<UBoneProxy::ETransformType>> VisibleTransform = TransformTypes;
+	
+	TSharedPtr<SSegmentedControl<UBoneProxy::ETransformType>> TransformChoiceWidget =
+		SSegmentedControl<UBoneProxy::ETransformType>::Create(
+			TransformTypes,
+			ButtonLabels,
+			ButtonTooltips,
+			VisibleTransform
+		);
 
 	// use a static shared ref so that all views retain these settings
 	static TSharedRef<TArray<UBoneProxy::ETransformType>> VisibleTransforms =
@@ -186,25 +206,7 @@ void FBoneProxyDetailsCustomization::CustomizeDetails(IDetailLayoutBuilder& Deta
 		.HAlign(HAlign_Left)
 		.VAlign(VAlign_Center)
 		[
-			SAssignNew(TransformChoiceWidget, SSegmentedControl<UBoneProxy::ETransformType>)
-			.SupportsMultiSelection(true)
-			.Values_Lambda([]()
-			{
-				return VisibleTransforms.Get();
-			})
-			.OnValuesChanged_Lambda([](TArray<UBoneProxy::ETransformType> Values)
-			{
-				VisibleTransforms.Get() = Values;
-			})
-			+ SSegmentedControl<UBoneProxy::ETransformType>::Slot(UBoneProxy::TransformType_Bone)
-			.Text(ButtonLabels[0])
-			.ToolTip(ButtonTooltips[0])
-			+ SSegmentedControl<UBoneProxy::ETransformType>::Slot(UBoneProxy::TransformType_Reference)
-			.Text(ButtonLabels[1])
-			.ToolTip(ButtonTooltips[1])
-			+ SSegmentedControl<UBoneProxy::ETransformType>::Slot(UBoneProxy::TransformType_Mesh)
-			.Text(ButtonLabels[2])
-			.ToolTip(ButtonTooltips[2])
+			TransformChoiceWidget.ToSharedRef()
 		]
 	];
 
@@ -276,89 +278,163 @@ void FBoneProxyDetailsCustomization::CustomizeDetails(IDetailLayoutBuilder& Deta
 		.IsEnabled(bIsEditingEnabled)
 		.DisplayRelativeWorld(bIsEditingEnabled)
 		.DisplayScaleLock(bIsEditingEnabled)
-		.OnGetNumericValue_Static(&UBoneProxy::GetMultiNumericValue, TransformType, BoneProxiesView);
-
-		if(bIsEditingEnabled)
+		.OnGetNumericValue_Static(&UBoneProxy::GetMultiNumericValue, TransformType, BoneProxiesView)
+		.OnCopyToClipboard_Lambda([TransformType, BoneProxiesView](ESlateTransformComponent::Type InComponent)
 		{
-			TransformWidgetArgs.OnNumericValueChanged_Static(&UBoneProxy::OnMultiNumericValueCommitted, ETextCommit::Default, TransformType, BoneProxiesView, false);
-			TransformWidgetArgs.OnNumericValueCommitted_Static (&UBoneProxy::OnMultiNumericValueCommitted, TransformType, BoneProxiesView, true);
-		}
-		else
-		{
-			TransformWidgetArgs._OnNumericValueChanged.Unbind();
-			TransformWidgetArgs._OnNumericValueCommitted.Unbind();
-		}
-		
-		IDetailGroup& Group = CategoryBuilder.AddGroup(*ButtonLabels[TransformIndex].ToString(), ButtonLabels[TransformIndex], false, true);
-		Group.HeaderRow()
-		.Visibility( TAttribute<EVisibility>::CreateLambda([TransformChoiceWidget, TransformType]() -> EVisibility
-		{
-			return TransformChoiceWidget->HasValue(TransformType) ? EVisibility::Visible : EVisibility::Collapsed;
-		}))
-		.NameContent()
-		[
-			SNew(STextBlock)
-			.Font(IDetailLayoutBuilder::GetDetailFont())
-			.Text(ButtonLabels[TransformIndex])
-			.ToolTipText(ButtonTooltips[TransformIndex])
-		];
+			if(BoneProxiesView.Num() == 0)
+			{
+				return;
+			}
+			
+			if(TransformType == UBoneProxy::TransformType_Bone)
+			{
+				FString Content;
+				UBoneProxy* BoneProxy = BoneProxiesView[0];
+				
+				switch(InComponent)
+				{
+					case ESlateTransformComponent::Location:
+					{
+						const FVector Data = BoneProxy->Location;
+						TBaseStructure<FVector>::Get()->ExportText(Content, &Data, &Data, nullptr, PPF_None, nullptr);
+						break;
+					}
+					case ESlateTransformComponent::Rotation:
+					{
+						const FRotator Data = BoneProxy->Rotation;
+						TBaseStructure<FRotator>::Get()->ExportText(Content, &Data, &Data, nullptr, PPF_None, nullptr);
+						break;
+					}
+					case ESlateTransformComponent::Scale:
+					{
+						const FVector Data = BoneProxy->Scale;
+						TBaseStructure<FVector>::Get()->ExportText(Content, &Data, &Data, nullptr, PPF_None, nullptr);
+						break;
+					}
+					case ESlateTransformComponent::Max:
+					default:
+					{
+						const FEulerTransform Data(BoneProxy->Location, BoneProxy->Rotation, BoneProxy->Scale);
+						TBaseStructure<FEulerTransform>::Get()->ExportText(Content, &Data, &Data, nullptr, PPF_None, nullptr);
+						break;
+					}
+				}
 
-		IDetailPropertyRow& LocationPropertyRow = Group.AddPropertyRow(Properties[PropertyIndex++]);
-		if(bIsEditingEnabled)
+				if(!Content.IsEmpty())
+				{
+					FPlatformApplicationMisc::ClipboardCopy(*Content);
+				}
+			}
+		})
+		.OnPasteFromClipboard_Lambda([TransformType, BoneProxiesView](ESlateTransformComponent::Type InComponent)
 		{
-			LocationPropertyRow.OverrideResetToDefault(FResetToDefaultOverride::Create(FIsResetToDefaultVisible::CreateSP(this, &FBoneProxyDetailsCustomization::IsResetLocationVisible, BoneProxiesView), FResetToDefaultHandler::CreateSP(this, &FBoneProxyDetailsCustomization::HandleResetLocation, BoneProxiesView)));
-		}
+			FString Content;
+			FPlatformApplicationMisc::ClipboardPaste(Content);
 
-		LocationPropertyRow.CustomWidget()
-		.NameContent()
-		.HAlign(HAlign_Fill)
-		[
-			SAdvancedTransformInputBox<FEulerTransform>::ConstructLabel(TransformWidgetArgs, ESlateTransformComponent::Location)
-		]
-		.ValueContent()
-		.MinDesiredWidth(BoneProxyCustomizationConstants::ItemWidth * 3.0f)
-		.MaxDesiredWidth(BoneProxyCustomizationConstants::ItemWidth * 3.0f)
-		[
-			SAdvancedTransformInputBox<FEulerTransform>::ConstructWidget(TransformWidgetArgs, ESlateTransformComponent::Location)
-		];
+			if(Content.IsEmpty() || TransformType != UBoneProxy::TransformType_Bone)
+			{
+				return;
+			}
 
-		IDetailPropertyRow& RotationPropertyRow = Group.AddPropertyRow(Properties[PropertyIndex++]);
-		if(bIsEditingEnabled)
+			for (UBoneProxy* BoneProxy : BoneProxiesView)
+			{
+				if (UDebugSkelMeshComponent* Component = BoneProxy->SkelMeshComponent.Get())
+				{
+					if (FAnimNode_ModifyBone* ModifyBone = Component->PreviewInstance->FindModifiedBone(BoneProxy->BoneName))
+					{
+						class FBoneProxyDetailsCustomization : public FOutputDevice
+						{
+						public:
+
+							int32 NumErrors;
+
+							FBoneProxyDetailsCustomization()
+								: FOutputDevice()
+								, NumErrors(0)
+							{
+							}
+
+							virtual void Serialize(const TCHAR* V, ELogVerbosity::Type Verbosity, const class FName& Category) override
+							{
+								NumErrors++;
+							}
+						};
+
+						FBoneProxyDetailsCustomization ErrorPipe;
+										
+						switch(InComponent)
+						{
+							case ESlateTransformComponent::Location:
+							{
+								FVector Data = BoneProxy->Location;
+								TBaseStructure<FVector>::Get()->ImportText(*Content, &Data, nullptr, PPF_None, &ErrorPipe, TBaseStructure<FVector>::Get()->GetName(), true);
+								ModifyBone->Translation = Data;
+								break;
+							}
+							case ESlateTransformComponent::Rotation:
+							{
+								FRotator Data = BoneProxy->Rotation;
+								TBaseStructure<FRotator>::Get()->ImportText(*Content, &Data, nullptr, PPF_None, &ErrorPipe, TBaseStructure<FRotator>::Get()->GetName(), true);
+								ModifyBone->Rotation = Data;
+								break;
+							}
+							case ESlateTransformComponent::Scale:
+							{
+								FVector Data = BoneProxy->Scale;
+								TBaseStructure<FVector>::Get()->ImportText(*Content, &Data, nullptr, PPF_None, &ErrorPipe, TBaseStructure<FVector>::Get()->GetName(), true);
+								ModifyBone->Scale = Data;
+								break;
+							}
+							case ESlateTransformComponent::Max:
+							default:
+							{
+								FEulerTransform Data = FEulerTransform::Identity;
+								TBaseStructure<FEulerTransform>::Get()->ImportText(*Content, &Data, nullptr, PPF_None, &ErrorPipe, TBaseStructure<FEulerTransform>::Get()->GetName(), true);
+								ModifyBone->Translation = Data.GetLocation();
+								ModifyBone->Rotation = Data.Rotator();
+								ModifyBone->Scale = Data.GetScale3D();
+								break;
+							}
+						}
+					}
+				}
+			}
+		})
+		.DiffersFromDefault_Lambda([TransformType, BoneProxiesView](ESlateTransformComponent::Type InComponent) -> bool
 		{
-			RotationPropertyRow.OverrideResetToDefault(FResetToDefaultOverride::Create(FIsResetToDefaultVisible::CreateSP(this, &FBoneProxyDetailsCustomization::IsResetRotationVisible, BoneProxiesView), FResetToDefaultHandler::CreateSP(this, &FBoneProxyDetailsCustomization::HandleResetRotation, BoneProxiesView)));
-		}
+			for(UBoneProxy* BoneProxy : BoneProxiesView)
+			{
+				if(BoneProxy->DiffersFromDefault(InComponent, TransformType))
+				{
+					return true;
+				}
+			}
+			return false;
+		})
+		.OnResetToDefault_Lambda([TransformType, BoneProxiesView](ESlateTransformComponent::Type InComponent)
+		{
+			for(UBoneProxy* BoneProxy : BoneProxiesView)
+			{
+				BoneProxy->ResetToDefault(InComponent, TransformType);
+			}
+		});
 
-		RotationPropertyRow.CustomWidget()
-		.NameContent()
-		.HAlign(HAlign_Fill)
-		[
-			SAdvancedTransformInputBox<FEulerTransform>::ConstructLabel(TransformWidgetArgs, ESlateTransformComponent::Rotation)
-		]
-		.ValueContent()
-		.MinDesiredWidth(BoneProxyCustomizationConstants::ItemWidth * 3.0f)
-		.MaxDesiredWidth(BoneProxyCustomizationConstants::ItemWidth * 3.0f)
-		[
-			SAdvancedTransformInputBox<FEulerTransform>::ConstructWidget(TransformWidgetArgs, ESlateTransformComponent::Rotation)
-		];
+		TransformWidgetArgs.Visibility_Lambda([TransformChoiceWidget, TransformType]() -> EVisibility
+		{
+			return TransformChoiceWidget->HasValue(TransformType) ? EVisibility::Visible : EVisibility::Collapsed;
+		});
 
-		IDetailPropertyRow& ScalePropertyRow = Group.AddPropertyRow(Properties[PropertyIndex++]);
 		if(bIsEditingEnabled)
 		{
-			ScalePropertyRow.OverrideResetToDefault(FResetToDefaultOverride::Create(FIsResetToDefaultVisible::CreateSP(this, &FBoneProxyDetailsCustomization::IsResetScaleVisible, BoneProxiesView), FResetToDefaultHandler::CreateSP(this, &FBoneProxyDetailsCustomization::HandleResetScale, BoneProxiesView)));
+			TransformWidgetArgs.OnNumericValueChanged_Static(&UBoneProxy::OnMultiNumericValueCommitted, ETextCommit::Default, TransformType, BoneProxiesView, false);
+			TransformWidgetArgs.OnNumericValueCommitted_Static (&UBoneProxy::OnMultiNumericValueCommitted, TransformType, BoneProxiesView, true);
 		}
 
-		ScalePropertyRow.CustomWidget()
-		.NameContent()
-		.HAlign(HAlign_Fill)
-		[
-			SAdvancedTransformInputBox<FEulerTransform>::ConstructLabel(TransformWidgetArgs, ESlateTransformComponent::Scale)
-		]
-		.ValueContent()
-		.MinDesiredWidth(BoneProxyCustomizationConstants::ItemWidth * 3.0f)
-		.MaxDesiredWidth(BoneProxyCustomizationConstants::ItemWidth * 3.0f)
-		[
-			SAdvancedTransformInputBox<FEulerTransform>::ConstructWidget(TransformWidgetArgs, ESlateTransformComponent::Scale)
-		];
+		SAdvancedTransformInputBox<FEulerTransform>::ConstructGroupedTransformRows(
+			CategoryBuilder, 
+			ButtonLabels[TransformIndex], 
+			ButtonTooltips[TransformIndex], 
+			TransformWidgetArgs);
 	}
 }
 
diff --git a/Engine/Source/Editor/SkeletonEditor/Public/BoneProxy.h b/Engine/Source/Editor/SkeletonEditor/Public/BoneProxy.h
--- a/Engine/Source/Editor/SkeletonEditor/Public/BoneProxy.h
+++ b/Engine/Source/Editor/SkeletonEditor/Public/BoneProxy.h
@@ -141,4 +141,10 @@ public:
 		ETransformType TransformType,
 		TArrayView<UBoneProxy*> BoneProxies,
 		bool bIsCommit);
+
+	// Returns true if a given transform component differs from the reference
+	bool DiffersFromDefault(ESlateTransformComponent::Type Component, ETransformType TransformType) const;
+
+	// Resets the bone transform's component to its reference
+	void ResetToDefault(ESlateTransformComponent::Type InComponent, ETransformType TransformType);
 };
