Converted a number of animation editor menus & toolbars to use tool menus

Also applied asset permissions to various UI sections

diff --git a/Engine/Plugins/Animation/ControlRig/Source/ControlRigEditor/Private/ControlRigEditorModule.cpp b/Engine/Plugins/Animation/ControlRig/Source/ControlRigEditor/Private/ControlRigEditorModule.cpp
--- a/Engine/Plugins/Animation/ControlRig/Source/ControlRigEditor/Private/ControlRigEditorModule.cpp
+++ b/Engine/Plugins/Animation/ControlRig/Source/ControlRigEditor/Private/ControlRigEditorModule.cpp
@@ -125,6 +125,7 @@
 #include "ControlRigSpaceChannelEditors.h"
 #include "UserDefinedStructure/UserDefinedStructEditorData.h"
 #include "UObject/FieldIterator.h"
+#include "AnimationToolMenuContext.h"
 #include "RigVMModel/Nodes/RigVMAggregateNode.h"
 #include "RigVMUserWorkflowRegistry.h"
 #include "Units/ControlRigNodeWorkflow.h"
@@ -201,12 +202,7 @@ void FControlRigEditorModule::StartupModule()
 	SequencerModule.RegisterChannelInterface<FMovieSceneControlRigSpaceChannel>();
 	ControlRigParameterTrackCreateEditorHandle = SequencerModule.RegisterTrackEditor(FOnCreateTrackEditor::CreateStatic(&FControlRigParameterTrackEditor::CreateTrackEditor));
 
-	//Register Animation Toolbar Extender
-	IAnimationEditorModule& AnimationEditorModule = FModuleManager::Get().LoadModuleChecked<IAnimationEditorModule>("AnimationEditor");
-	auto& ToolbarExtenders = AnimationEditorModule.GetAllAnimationEditorToolbarExtenders();
-
-	ToolbarExtenders.Add(IAnimationEditorModule::FAnimationEditorToolbarExtender::CreateRaw(this, &FControlRigEditorModule::GetAnimationEditorToolbarExtender));
-	AnimationEditorExtenderHandle = ToolbarExtenders.Last().GetHandle();
+	AddControlRigExtenderToToolMenu("AssetEditor.AnimationEditor.ToolBar");
 
 	FEditorModeRegistry::Get().RegisterMode<FControlRigEditMode>(
 		FControlRigEditMode::ModeName,
@@ -323,13 +319,6 @@ void FControlRigEditorModule::ShutdownModule()
 		}
 	}
 
-	IAnimationEditorModule* AnimationEditorModule = FModuleManager::Get().GetModulePtr<IAnimationEditorModule>("AnimationEditor");
-	if (AnimationEditorModule)
-	{
-		typedef IAnimationEditorModule::FAnimationEditorToolbarExtender DelegateType;
-		AnimationEditorModule->GetAllAnimationEditorToolbarExtenders().RemoveAll([=](const DelegateType& In) { return In.GetHandle() == AnimationEditorExtenderHandle; });
-	}
-
 	if (UObjectInitialized())
 	{
 		for(const int32 WorkflowHandle : WorkflowHandles)
@@ -343,34 +332,6 @@ void FControlRigEditorModule::ShutdownModule()
 	WorkflowHandles.Reset();
 }
 
-TSharedRef<FExtender> FControlRigEditorModule::GetAnimationEditorToolbarExtender(const TSharedRef<FUICommandList> CommandList, TSharedRef<IAnimationEditor> InAnimationEditor)
-{
-	TSharedRef<FExtender> Extender = MakeShareable(new FExtender);
-
-	USkeleton* Skeleton = InAnimationEditor->GetPersonaToolkit()->GetSkeleton();
-	USkeletalMesh* SkeletalMesh = InAnimationEditor->GetPersonaToolkit()->GetPreviewMesh();
-	if (!SkeletalMesh) //if no preview mesh just get normal mesh
-	{
-		SkeletalMesh = InAnimationEditor->GetPersonaToolkit()->GetMesh();
-	}
-	if (Skeleton && SkeletalMesh)
-	{
-		UAnimSequence* AnimSequence = Cast<UAnimSequence>(InAnimationEditor->GetPersonaToolkit()->GetAnimationAsset());
-		if (AnimSequence)
-		{
-			Extender->AddToolBarExtension(
-				"Asset",
-				EExtensionHook::After,
-				CommandList,
-				FToolBarExtensionDelegate::CreateRaw(this, &FControlRigEditorModule::HandleAddControlRigExtenderToToolbar, TWeakPtr<IAnimationEditor>(InAnimationEditor))
-			);
-		}
-	}
-
-	return Extender;
-}
-
-
 TSharedRef< SWidget > FControlRigEditorModule::GenerateAnimationMenu(TWeakPtr<IAnimationEditor> InAnimationEditor)
 {
 	const bool bShouldCloseWindowAfterMenuSelection = true;
@@ -891,14 +852,36 @@ void FControlRigEditorModule::OpenLevelSequence(UAnimSequence* AnimSequence)
 	}
 }
 
-void FControlRigEditorModule::HandleAddControlRigExtenderToToolbar(FToolBarBuilder& ParentToolbarBuilder, TWeakPtr<IAnimationEditor> InAnimationEditor)
+void FControlRigEditorModule::AddControlRigExtenderToToolMenu(FName InToolMenuName)
 {
-	ParentToolbarBuilder.AddComboButton(
-		FUIAction(),
-		FOnGetContent::CreateRaw(this, &FControlRigEditorModule::GenerateAnimationMenu, InAnimationEditor),
-		LOCTEXT("EditInSequencer", "Edit in Sequencer"),
-		LOCTEXT("EditInSequencer_Tooltip", "Edit this Anim Sequence In Sequencer."),
-		FSlateIcon(FEditorStyle::GetStyleSetName(), "Persona.EditInSequencer")
+	FToolMenuOwnerScoped OwnerScoped(this);
+
+	UToolMenu* ToolMenu = UToolMenus::Get()->ExtendMenu(InToolMenuName);
+	ToolMenu->AddMenuEntry("Sequencer",
+		FToolMenuEntry::InitComboButton(
+			"EditInSequencer",
+			FToolUIActionChoice(),
+			FNewToolMenuChoice(
+				FNewToolMenuDelegate::CreateLambda([this](UToolMenu* InNewToolMenu)
+				{
+					if (UAnimationToolMenuContext* MenuContext = InNewToolMenu->FindContext<UAnimationToolMenuContext>())
+					{
+						InNewToolMenu->AddMenuEntry(
+							"EditInSequencer", 
+							FToolMenuEntry::InitWidget(
+								"EditInSequencerMenu", 
+								GenerateAnimationMenu(MenuContext->AnimationEditor),
+								FText::GetEmpty(),
+								true, false, true
+							)
+						);
+					}
+				})
+			),
+			LOCTEXT("EditInSequencer", "Edit in Sequencer"),
+			LOCTEXT("EditInSequencer_Tooltip", "Edit this Anim Sequence In Sequencer."),
+			FSlateIcon(FEditorStyle::GetStyleSetName(), "Persona.EditInSequencer")
+		)
 	);
 }
 
diff --git a/Engine/Plugins/Animation/ControlRig/Source/ControlRigEditor/Private/ControlRigEditorModule.h b/Engine/Plugins/Animation/ControlRig/Source/ControlRigEditor/Private/ControlRigEditorModule.h
--- a/Engine/Plugins/Animation/ControlRig/Source/ControlRigEditor/Private/ControlRigEditorModule.h
+++ b/Engine/Plugins/Animation/ControlRig/Source/ControlRigEditor/Private/ControlRigEditorModule.h
@@ -40,8 +40,7 @@ public:
 	virtual TSharedPtr<FExtensibilityManager> GetToolBarExtensibilityManager() override { return ToolBarExtensibilityManager; }
 	
 	/** Animation Toolbar Extender*/
-	TSharedRef<FExtender> GetAnimationEditorToolbarExtender(const TSharedRef<FUICommandList> CommandList, TSharedRef<IAnimationEditor> InAnimationEditor);
-	void HandleAddControlRigExtenderToToolbar(FToolBarBuilder& ParentToolbarBuilder, TWeakPtr<IAnimationEditor> InAnimationEditor);
+	void AddControlRigExtenderToToolMenu(FName InToolMenuName);
 	TSharedRef< SWidget > GenerateAnimationMenu(TWeakPtr<IAnimationEditor> InAnimationEditor);
 	void ToggleIsDrivenByLevelSequence(UAnimSequence* AnimSequence)const;
 	bool IsDrivenByLevelSequence(UAnimSequence* AnimSequence)const;
@@ -88,9 +87,6 @@ private:
 	FDelegateHandle ReconstructAllNodesDelegateHandle;
 	FDelegateHandle BlueprintVariableCustomizationHandle;
 
-	/** Animation Editor Handle*/
-	FDelegateHandle AnimationEditorExtenderHandle;
-
 	/** Param to hold Filter Result to pass to Filter*/
 	bool bFilterAssetBySkeleton;
 
diff --git a/Engine/Source/Developer/MeshUtilities/Private/MeshUtilities.cpp b/Engine/Source/Developer/MeshUtilities/Private/MeshUtilities.cpp
--- a/Engine/Source/Developer/MeshUtilities/Private/MeshUtilities.cpp
+++ b/Engine/Source/Developer/MeshUtilities/Private/MeshUtilities.cpp
@@ -120,6 +120,7 @@
 #include "MeshMergeModule.h"
 
 #include "Animation/SkinWeightProfile.h"
+#include "PersonaToolMenuContext.h"
 
 DEFINE_LOG_CATEGORY(LogMeshUtilities);
 /*------------------------------------------------------------------------------
@@ -5555,18 +5556,6 @@ void FMeshUtilities::StartupModule()
 			{
 				AddLevelViewportMenuExtender();
 			}
-			else if (InModuleName == "AnimationBlueprintEditor")
-			{
-				AddAnimationBlueprintEditorToolbarExtender();
-			}
-			else if (InModuleName == "AnimationEditor")
-			{
-				AddAnimationEditorToolbarExtender();
-			}
-			else if (InModuleName == "SkeletonEditor")
-			{
-				AddSkeletonEditorToolbarExtender();
-			}
 		}
 	});
 
@@ -5589,43 +5578,16 @@ void FMeshUtilities::ShutdownModule()
 	}
 
 	RemoveLevelViewportMenuExtender();
-	RemoveAnimationBlueprintEditorToolbarExtender();
-	RemoveAnimationEditorToolbarExtender();
-	RemoveSkeletonEditorToolbarExtender();
 	FModuleManager::Get().OnModulesChanged().Remove(ModuleLoadedDelegateHandle);
 	VersionString.Empty();
 }
 
 void FMeshUtilities::RegisterMenus()
 {
-	FToolMenuOwnerScoped OwnerScoped(this);
-
-	static auto AddMakeStaticMeshToolbarButton = [this](FToolMenuSection& InSection, const FToolMenuExecuteAction& InAction)
-	{
-		InSection.AddEntry(FToolMenuEntry::InitToolBarButton(
-			"MakeStaticMesh",
-			InAction,
-			LOCTEXT("MakeStaticMesh", "Make Static Mesh"),
-			LOCTEXT("MakeStaticMeshTooltip", "Make a new static mesh out of the preview's current pose."),
-			FSlateIcon("EditorStyle", "Persona.ConvertToStaticMesh")
-		));
-	};
-
-	{
-		UToolMenu* Toolbar = UToolMenus::Get()->ExtendMenu("AssetEditor.SkeletalMeshEditor.ToolBar");
-		FToolMenuSection& Section = Toolbar->FindOrAddSection("SkeletalMesh");
-		AddMakeStaticMeshToolbarButton(Section, FToolMenuExecuteAction::CreateLambda([this](const FToolMenuContext& InMenuContext)
-		{
-			USkeletalMeshToolMenuContext* Context = InMenuContext.FindContext<USkeletalMeshToolMenuContext>();
-			if (Context && Context->SkeletalMeshEditor.IsValid())
-			{
-				if (UMeshComponent* MeshComponent = Context->SkeletalMeshEditor.Pin()->GetPersonaToolkit()->GetPreviewMeshComponent())
-				{
-					ConvertMeshesToStaticMesh(TArray<UMeshComponent*>({ MeshComponent }), MeshComponent->GetComponentToWorld());
-				}
-			}
-		}));
-	}
+	AddMakeStaticMeshEntryToToolMenu("AssetEditor.AnimationEditor.ToolBar");
+	AddMakeStaticMeshEntryToToolMenu("AssetEditor.AnimationBlueprintEditor.ToolBar");
+	AddMakeStaticMeshEntryToToolMenu("AssetEditor.SkeletalMeshEditor.ToolBar");
+	AddMakeStaticMeshEntryToToolMenu("AssetEditor.SkeletonEditor.ToolBar");
 }
 
 bool FMeshUtilities::GenerateUniqueUVsForSkeletalMesh(const FSkeletalMeshLODModel& LODModel, int32 TextureResolution, TArray<FVector2f>& OutTexCoords) const
@@ -6120,142 +6082,30 @@ void FMeshUtilities::CreateImportDataFromLODModel(USkeletalMesh* SkeletalMesh) c
 	}
 }
 
-void FMeshUtilities::AddAnimationBlueprintEditorToolbarExtender()
-{
-	IAnimationBlueprintEditorModule& AnimationBlueprintEditorModule = FModuleManager::Get().LoadModuleChecked<IAnimationBlueprintEditorModule>("AnimationBlueprintEditor");
-	auto& ToolbarExtenders = AnimationBlueprintEditorModule.GetAllAnimationBlueprintEditorToolbarExtenders();
-
-	ToolbarExtenders.Add(IAnimationBlueprintEditorModule::FAnimationBlueprintEditorToolbarExtender::CreateRaw(this, &FMeshUtilities::GetAnimationBlueprintEditorToolbarExtender));
-	AnimationBlueprintEditorExtenderHandle = ToolbarExtenders.Last().GetHandle();
-}
-
-void FMeshUtilities::RemoveAnimationBlueprintEditorToolbarExtender()
-{
-	IAnimationBlueprintEditorModule* AnimationBlueprintEditorModule = FModuleManager::Get().GetModulePtr<IAnimationBlueprintEditorModule>("AnimationBlueprintEditor");
-	if (AnimationBlueprintEditorModule)
-	{
-		typedef IAnimationBlueprintEditorModule::FAnimationBlueprintEditorToolbarExtender DelegateType;
-		AnimationBlueprintEditorModule->GetAllAnimationBlueprintEditorToolbarExtenders().RemoveAll([=](const DelegateType& In) { return In.GetHandle() == AnimationBlueprintEditorExtenderHandle; });
-	}
-}
-
-TSharedRef<FExtender> FMeshUtilities::GetAnimationBlueprintEditorToolbarExtender(const TSharedRef<FUICommandList> CommandList, TSharedRef<IAnimationBlueprintEditor> InAnimationBlueprintEditor)
+void FMeshUtilities::AddMakeStaticMeshEntryToToolMenu(FName InToolMenuName)
 {
-	TSharedRef<FExtender> Extender = MakeShareable(new FExtender);
-
-	if(InAnimationBlueprintEditor->GetBlueprintObj() && InAnimationBlueprintEditor->GetBlueprintObj()->BlueprintType != BPTYPE_Interface)
-	{
-		UMeshComponent* MeshComponent = InAnimationBlueprintEditor->GetPersonaToolkit()->GetPreviewMeshComponent();
-
-		Extender->AddToolBarExtension(
-			"Asset",
-			EExtensionHook::After,
-			CommandList,
-			FToolBarExtensionDelegate::CreateRaw(this, &FMeshUtilities::HandleAddSkeletalMeshActionExtenderToToolbar, MeshComponent)
-		);
-	}
-
-	return Extender;
-}
-
-void FMeshUtilities::AddAnimationEditorToolbarExtender()
-{
-	IAnimationEditorModule& AnimationEditorModule = FModuleManager::Get().LoadModuleChecked<IAnimationEditorModule>("AnimationEditor");
-	auto& ToolbarExtenders = AnimationEditorModule.GetAllAnimationEditorToolbarExtenders();
-
-	ToolbarExtenders.Add(IAnimationEditorModule::FAnimationEditorToolbarExtender::CreateRaw(this, &FMeshUtilities::GetAnimationEditorToolbarExtender));
-	AnimationEditorExtenderHandle = ToolbarExtenders.Last().GetHandle();
-}
-
-void FMeshUtilities::RemoveAnimationEditorToolbarExtender()
-{
-	IAnimationEditorModule* AnimationEditorModule = FModuleManager::Get().GetModulePtr<IAnimationEditorModule>("AnimationEditor");
-	if (AnimationEditorModule)
-	{
-		typedef IAnimationEditorModule::FAnimationEditorToolbarExtender DelegateType;
-		AnimationEditorModule->GetAllAnimationEditorToolbarExtenders().RemoveAll([=](const DelegateType& In) { return In.GetHandle() == AnimationEditorExtenderHandle; });
-	}
-}
-
-TSharedRef<FExtender> FMeshUtilities::GetAnimationEditorToolbarExtender(const TSharedRef<FUICommandList> CommandList, TSharedRef<IAnimationEditor> InAnimationEditor)
-{
-	TSharedRef<FExtender> Extender = MakeShareable(new FExtender);
-
-	UMeshComponent* MeshComponent = InAnimationEditor->GetPersonaToolkit()->GetPreviewMeshComponent();
-
-	Extender->AddToolBarExtension(
-		"Asset",
-		EExtensionHook::After,
-		CommandList,
-		FToolBarExtensionDelegate::CreateRaw(this, &FMeshUtilities::HandleAddSkeletalMeshActionExtenderToToolbar, MeshComponent)
-	);
-
-	return Extender;
-}
-
-TSharedRef<FExtender> FMeshUtilities::GetSkeletalMeshEditorToolbarExtender(const TSharedRef<FUICommandList> CommandList, TSharedRef<ISkeletalMeshEditor> InSkeletalMeshEditor)
-{
-	TSharedRef<FExtender> Extender = MakeShareable(new FExtender);
-
-	UMeshComponent* MeshComponent = InSkeletalMeshEditor->GetPersonaToolkit()->GetPreviewMeshComponent();
-
-	Extender->AddToolBarExtension(
-		"Asset",
-		EExtensionHook::After,
-		CommandList,
-		FToolBarExtensionDelegate::CreateRaw(this, &FMeshUtilities::HandleAddSkeletalMeshActionExtenderToToolbar, MeshComponent)
-	);
-
-	return Extender;
-}
-
-void FMeshUtilities::AddSkeletonEditorToolbarExtender()
-{
-	ISkeletonEditorModule& SkeletonEditorModule = FModuleManager::Get().LoadModuleChecked<ISkeletonEditorModule>("SkeletonEditor");
-	auto& ToolbarExtenders = SkeletonEditorModule.GetAllSkeletonEditorToolbarExtenders();
-
-	ToolbarExtenders.Add(ISkeletonEditorModule::FSkeletonEditorToolbarExtender::CreateRaw(this, &FMeshUtilities::GetSkeletonEditorToolbarExtender));
-	SkeletonEditorExtenderHandle = ToolbarExtenders.Last().GetHandle();
-}
-
-void FMeshUtilities::RemoveSkeletonEditorToolbarExtender()
-{
-	ISkeletonEditorModule* SkeletonEditorModule = FModuleManager::Get().GetModulePtr<ISkeletonEditorModule>("SkeletonEditor");
-	if (SkeletonEditorModule)
-	{
-		typedef ISkeletonEditorModule::FSkeletonEditorToolbarExtender DelegateType;
-		SkeletonEditorModule->GetAllSkeletonEditorToolbarExtenders().RemoveAll([=](const DelegateType& In) { return In.GetHandle() == SkeletonEditorExtenderHandle; });
-	}
-}
-
-TSharedRef<FExtender> FMeshUtilities::GetSkeletonEditorToolbarExtender(const TSharedRef<FUICommandList> CommandList, TSharedRef<ISkeletonEditor> InSkeletonEditor)
-{
-	TSharedRef<FExtender> Extender = MakeShareable(new FExtender);
-
-	UMeshComponent* MeshComponent = InSkeletonEditor->GetPersonaToolkit()->GetPreviewMeshComponent();
-
-	Extender->AddToolBarExtension(
-		"Asset",
-		EExtensionHook::After,
-		CommandList,
-		FToolBarExtensionDelegate::CreateRaw(this, &FMeshUtilities::HandleAddSkeletalMeshActionExtenderToToolbar, MeshComponent)
-	);
-
-	return Extender;
-}
-
+	FToolMenuOwnerScoped OwnerScoped(this);
 
-void FMeshUtilities::HandleAddSkeletalMeshActionExtenderToToolbar(FToolBarBuilder& ParentToolbarBuilder, UMeshComponent* InMeshComponent)
-{
-	ParentToolbarBuilder.AddToolBarButton(
-		FUIAction(FExecuteAction::CreateLambda([this, InMeshComponent]()
-		{
-			ConvertMeshesToStaticMesh(TArray<UMeshComponent*>({ InMeshComponent }), InMeshComponent->GetComponentToWorld());
-		})),
-		NAME_None,
-		LOCTEXT("MakeStaticMesh", "Make Static Mesh"),
-		LOCTEXT("MakeStaticMeshTooltip", "Make a new static mesh out of the preview's current pose."),
-		FSlateIcon("EditorStyle", "Persona.ConvertToStaticMesh")
+	UToolMenu* ToolMenu = UToolMenus::Get()->ExtendMenu(InToolMenuName);
+	ToolMenu->AddMenuEntry("SkeletalMesh",
+		FToolMenuEntry::InitToolBarButton(
+			"MakeStaticMesh",
+			FToolUIActionChoice(
+				FToolUIAction(FToolMenuExecuteAction::CreateLambda([this](const FToolMenuContext& InContext)
+				{
+					if(UPersonaToolMenuContext* MenuContext = InContext.FindContext<UPersonaToolMenuContext>())
+					{
+						if(UDebugSkelMeshComponent* MeshComponent = MenuContext->GetPreviewMeshComponent())
+						{
+							ConvertMeshesToStaticMesh(TArray<UMeshComponent*>({ MeshComponent }), MeshComponent->GetComponentToWorld());
+						}
+					}
+				}))
+			),
+			LOCTEXT("MakeStaticMesh", "Make Static Mesh"),
+			LOCTEXT("MakeStaticMeshTooltip", "Make a new static mesh out of the preview's current pose."),
+			FSlateIcon("EditorStyle", "Persona.ConvertToStaticMesh")
+		)
 	);
 }
 
diff --git a/Engine/Source/Developer/MeshUtilities/Private/MeshUtilitiesPrivate.h b/Engine/Source/Developer/MeshUtilities/Private/MeshUtilitiesPrivate.h
--- a/Engine/Source/Developer/MeshUtilities/Private/MeshUtilitiesPrivate.h
+++ b/Engine/Source/Developer/MeshUtilities/Private/MeshUtilitiesPrivate.h
@@ -164,27 +164,7 @@ private:
 	friend class FStaticMeshUtilityBuilder;
 
 protected:
-	void AddAnimationBlueprintEditorToolbarExtender();
-
-	void RemoveAnimationBlueprintEditorToolbarExtender();
-
-	TSharedRef<FExtender> GetAnimationBlueprintEditorToolbarExtender(const TSharedRef<FUICommandList> CommandList, TSharedRef<IAnimationBlueprintEditor> InAnimationBlueprintEditor);
-
-	void AddAnimationEditorToolbarExtender();
-
-	void RemoveAnimationEditorToolbarExtender();
-
-	TSharedRef<FExtender> GetAnimationEditorToolbarExtender(const TSharedRef<FUICommandList> CommandList, TSharedRef<IAnimationEditor> InAnimationEditor);
-
-	TSharedRef<FExtender> GetSkeletalMeshEditorToolbarExtender(const TSharedRef<FUICommandList> CommandList, TSharedRef<ISkeletalMeshEditor> InSkeletalMeshEditor);
-
-	void AddSkeletonEditorToolbarExtender();
-
-	void RemoveSkeletonEditorToolbarExtender();
-
-	TSharedRef<FExtender> GetSkeletonEditorToolbarExtender(const TSharedRef<FUICommandList> CommandList, TSharedRef<ISkeletonEditor> InSkeletonEditor);
-
-	void HandleAddSkeletalMeshActionExtenderToToolbar(FToolBarBuilder& ParentToolbarBuilder, UMeshComponent* MeshComponent);
+	void AddMakeStaticMeshEntryToToolMenu(FName InToolbarName);
 
 	void AddLevelViewportMenuExtender();
 
diff --git a/Engine/Source/Editor/AnimationEditor/Private/AnimationEditor.cpp b/Engine/Source/Editor/AnimationEditor/Private/AnimationEditor.cpp
--- a/Engine/Source/Editor/AnimationEditor/Private/AnimationEditor.cpp
+++ b/Engine/Source/Editor/AnimationEditor/Private/AnimationEditor.cpp
@@ -54,6 +54,9 @@
 #include "IPersonaEditorModeManager.h"
 #include "Toolkits/AssetEditorToolkit.h"
 #include "PersonaToolMenuContext.h"
+#include "ToolMenus.h"
+#include "ToolMenuMisc.h"
+#include "AnimationToolMenuContext.h"
 
 const FName AnimationEditorAppIdentifier = FName(TEXT("AnimationEditorApp"));
 
@@ -163,10 +166,13 @@ void FAnimationEditor::InitToolMenuContext(FToolMenuContext& MenuContext)
 {
 	FAssetEditorToolkit::InitToolMenuContext(MenuContext);
 
-	UPersonaToolMenuContext* Context = NewObject<UPersonaToolMenuContext>();
-	Context->SetToolkit(GetPersonaToolkit());
+	UAnimationToolMenuContext* AnimationToolMenuContext = NewObject<UAnimationToolMenuContext>();
+	AnimationToolMenuContext->AnimationEditor = SharedThis(this);
+	MenuContext.AddObject(AnimationToolMenuContext);
 
-	MenuContext.AddObject(Context);
+	UPersonaToolMenuContext* PersonaToolMenuContext = NewObject<UPersonaToolMenuContext>();
+	PersonaToolMenuContext->SetToolkit(GetPersonaToolkit());
+	MenuContext.AddObject(PersonaToolMenuContext);
 }
 
 void FAnimationEditor::Tick(float DeltaTime)
@@ -220,8 +226,64 @@ void FAnimationEditor::BindCommands()
 		FExecuteAction::CreateRaw(&GetPersonaToolkit()->GetPreviewScene().Get(), &IPersonaPreviewScene::TogglePlayback));
 }
 
+TSharedPtr<FAnimationEditor> FAnimationEditor::GetAnimationEditor(const FToolMenuContext& InMenuContext)
+{
+	if (UAnimationToolMenuContext* Context = InMenuContext.FindContext<UAnimationToolMenuContext>())
+	{
+		if (Context->AnimationEditor.IsValid())
+		{
+			return StaticCastSharedPtr<FAnimationEditor>(Context->AnimationEditor.Pin());
+		}
+	}
+
+	return TSharedPtr<FAnimationEditor>();
+}
+
 void FAnimationEditor::ExtendToolbar()
 {
+	FToolMenuOwnerScoped OwnerScoped(this);
+
+	// Add in Editor Specific functionality
+	FName ParentName;
+	static const FName MenuName = GetToolMenuToolbarName(ParentName);
+
+	UToolMenu* ToolMenu = UToolMenus::Get()->ExtendMenu(MenuName);
+	const FToolMenuInsert SectionInsertLocation("Asset", EToolMenuInsertType::After);
+
+	{
+		ToolMenu->AddDynamicSection("Persona", FNewToolBarDelegateLegacy::CreateLambda([](FToolBarBuilder& ToolbarBuilder, UToolMenu* InMenu)
+		{
+			TSharedPtr<FAnimationEditor> AnimationEditor = GetAnimationEditor(InMenu->Context);
+			if (AnimationEditor.IsValid() && AnimationEditor->PersonaToolkit.IsValid())
+			{
+				FPersonaModule& PersonaModule = FModuleManager::LoadModuleChecked<FPersonaModule>("Persona");
+				FPersonaModule::FCommonToolbarExtensionArgs Args;
+				Args.bPreviewAnimation = false;
+				Args.bReferencePose = false;
+				PersonaModule.AddCommonToolbarExtensions(ToolbarBuilder, AnimationEditor->PersonaToolkit.ToSharedRef(), Args);
+			}
+		}), SectionInsertLocation);
+	}
+
+	{
+		FToolMenuSection& AnimationSection = ToolMenu->AddSection("Animation", LOCTEXT("ToolbarAnimationSectionLabel", "Animation"), SectionInsertLocation);
+		AnimationSection.AddEntry(FToolMenuEntry::InitToolBarButton(FAnimationEditorCommands::Get().ReimportAnimation));
+		AnimationSection.AddEntry(FToolMenuEntry::InitToolBarButton(FAnimationEditorCommands::Get().ApplyCompression, LOCTEXT("Toolbar_ApplyCompression", "Apply Compression")));
+		AnimationSection.AddEntry(FToolMenuEntry::InitComboButton(
+			"ExportAsset",
+			FToolUIActionChoice(FUIAction()),
+			FNewToolMenuChoice(FOnGetContent::CreateSP(this, &FAnimationEditor::GenerateExportAssetMenu)),
+			LOCTEXT("ExportAsset_Label", "Export Asset"),
+			LOCTEXT("ExportAsset_ToolTip", "Export Assets for this skeleton."),
+			FSlateIcon(FEditorStyle::GetStyleSetName(), "Persona.ExportToFBX")
+		));
+	}
+
+	{
+		FToolMenuSection& EditingSection = ToolMenu->AddSection("Editing", LOCTEXT("ToolbarEditingSectionLabel", "Editing"), SectionInsertLocation);
+		EditingSection.AddEntry(FToolMenuEntry::InitToolBarButton(FAnimationEditorCommands::Get().SetKey, LOCTEXT("Toolbar_SetKey", "Key")));
+	}
+
 	// If the ToolbarExtender is valid, remove it before rebuilding it
 	if (ToolbarExtender.IsValid())
 	{
@@ -254,34 +316,6 @@ void FAnimationEditor::ExtendToolbar()
 		FToolBarExtensionDelegate::CreateLambda([this](FToolBarBuilder& ToolbarBuilder)
 		{
 			FPersonaModule& PersonaModule = FModuleManager::LoadModuleChecked<FPersonaModule>("Persona");
-			FPersonaModule::FCommonToolbarExtensionArgs Args;
-			Args.bPreviewAnimation = false;
-			Args.bReferencePose = false;
-			PersonaModule.AddCommonToolbarExtensions(ToolbarBuilder, PersonaToolkit.ToSharedRef(), Args);
-
-			ToolbarBuilder.BeginSection("Animation");
-			{
-				ToolbarBuilder.AddToolBarButton(FAnimationEditorCommands::Get().ReimportAnimation);
-				ToolbarBuilder.AddToolBarButton(FAnimationEditorCommands::Get().ApplyCompression, NAME_None, LOCTEXT("Toolbar_ApplyCompression", "Apply Compression"));
-
-				{
-					ToolbarBuilder.AddComboButton(
-						FUIAction(),
-						FOnGetContent::CreateSP(this, &FAnimationEditor::GenerateExportAssetMenu),
-						LOCTEXT("ExportAsset_Label", "Export Asset"),
-						LOCTEXT("ExportAsset_ToolTip", "Export Assets for this skeleton."),
-						FSlateIcon(FEditorStyle::GetStyleSetName(), "Persona.ExportToFBX")
-					);
-				}
-			}
-			ToolbarBuilder.EndSection();
-
-			ToolbarBuilder.BeginSection("Editing");
-			{
-				ToolbarBuilder.AddToolBarButton(FAnimationEditorCommands::Get().SetKey, NAME_None, LOCTEXT("Toolbar_SetKey", "Key"));
-			}
-			ToolbarBuilder.EndSection();
-
 			TSharedRef<class IAssetFamily> AssetFamily = PersonaModule.CreatePersonaAssetFamily(AnimationAsset);
 			AddToolbarWidget(PersonaModule.CreateAssetFamilyShortcutWidget(SharedThis(this), AssetFamily));
 		}	
@@ -292,35 +326,26 @@ void FAnimationEditor::ExtendMenu()
 {
 	MenuExtender = MakeShareable(new FExtender);
 
-	struct Local
-	{
-		static void AddAssetMenu(FMenuBuilder& MenuBuilder, FAnimationEditor* InAnimationEditor)
-		{
-			MenuBuilder.BeginSection("AnimationEditor", LOCTEXT("AnimationEditorAssetMenu_Animation", "Animation"));
-			{
-				MenuBuilder.AddMenuEntry(FAnimationEditorCommands::Get().ApplyCompression);
-
-				MenuBuilder.AddSubMenu(
-					LOCTEXT("ExportToFBX", "Export to FBX"),
-					LOCTEXT("ExportToFBX_ToolTip", "Export current animation to FBX"),
-					FNewMenuDelegate::CreateSP(InAnimationEditor, &FAnimationEditor::FillExportAssetMenu),
-					false,
-					FSlateIcon(FEditorStyle::GetStyleSetName(), "ClassIcon.")
-				);
-
-				MenuBuilder.AddMenuEntry(FAnimationEditorCommands::Get().AddLoopingInterpolation);
-				MenuBuilder.AddMenuEntry(FAnimationEditorCommands::Get().RemoveBoneTracks);
-			}
-			MenuBuilder.EndSection();
-		}
-	};
+	FToolMenuOwnerScoped OwnerScoped(this);
 
-	MenuExtender->AddMenuExtension(
-		"AssetEditorActions",
-		EExtensionHook::After,
-		GetToolkitCommands(),
-		FMenuExtensionDelegate::CreateStatic(&Local::AddAssetMenu, this)
-		);
+	// Add in Editor Specific functionality
+	UToolMenu* ToolMenu = UToolMenus::Get()->ExtendMenu("AssetEditor.AnimationEditor.MainMenu.Asset");
+	const FToolMenuInsert SectionInsertLocation("AssetEditorActions", EToolMenuInsertType::After);
+
+	FToolMenuSection& AnimationSection = ToolMenu->AddSection("AnimationEditor", LOCTEXT("AnimationEditorAssetMenu_Animation", "Animation"), SectionInsertLocation);
+	AnimationSection.AddEntry(FToolMenuEntry::InitMenuEntry(FAnimationEditorCommands::Get().ApplyCompression));
+	
+	AnimationSection.AddEntry(FToolMenuEntry::InitSubMenu(
+		"ExportAsset",
+		LOCTEXT("ExportAsset_Label", "Export Asset"),
+		LOCTEXT("ExportAsset_ToolTip", "Export Assets for this skeleton."),
+		FNewToolMenuChoice(FNewMenuDelegate::CreateSP(this, &FAnimationEditor::FillExportAssetMenu)),
+		false,
+		FSlateIcon(FEditorStyle::GetStyleSetName(), "Persona.ExportToFBX")
+	));
+
+	AnimationSection.AddEntry(FToolMenuEntry::InitMenuEntry(FAnimationEditorCommands::Get().AddLoopingInterpolation));
+	AnimationSection.AddEntry(FToolMenuEntry::InitMenuEntry(FAnimationEditorCommands::Get().RemoveBoneTracks));
 
 	AddMenuExtender(MenuExtender);
 
diff --git a/Engine/Source/Editor/AnimationEditor/Private/AnimationEditor.h b/Engine/Source/Editor/AnimationEditor/Private/AnimationEditor.h
--- a/Engine/Source/Editor/AnimationEditor/Private/AnimationEditor.h
+++ b/Engine/Source/Editor/AnimationEditor/Private/AnimationEditor.h
@@ -161,6 +161,8 @@ private:
 	TSharedPtr<SDockTab> OpenNewAnimationDocumentTab(UAnimationAsset* InAnimAsset);
 
 	bool RecordMeshToAnimation(USkeletalMeshComponent* PreviewComponent, UAnimSequence* NewAsset) const;
+
+	static TSharedPtr<FAnimationEditor> GetAnimationEditor(const FToolMenuContext& InMenuContext);
 public:
 	/** Multicast delegate fired on global undo/redo */
 	FSimpleMulticastDelegate OnLODChanged;
diff --git a/Engine/Source/Editor/AnimationEditor/Public/AnimationToolMenuContext.h b/Engine/Source/Editor/AnimationEditor/Public/AnimationToolMenuContext.h
new file mode 100644
--- /dev/null
+++ b/Engine/Source/Editor/AnimationEditor/Public/AnimationToolMenuContext.h
@@ -0,0 +1,18 @@
+// Copyright Epic Games, Inc. All Rights Reserved.
+
+#pragma once
+
+#include "CoreMinimal.h"
+#include "ToolMenuContext.h"
+#include "AnimationToolMenuContext.generated.h"
+
+class IAnimationEditor;
+
+UCLASS()
+class ANIMATIONEDITOR_API UAnimationToolMenuContext : public UObject
+{
+	GENERATED_BODY()
+
+public:
+	TWeakPtr<IAnimationEditor> AnimationEditor;
+};
\ No newline at end of file
diff --git a/Engine/Source/Editor/Persona/Private/PersonaModule.cpp b/Engine/Source/Editor/Persona/Private/PersonaModule.cpp
--- a/Engine/Source/Editor/Persona/Private/PersonaModule.cpp
+++ b/Engine/Source/Editor/Persona/Private/PersonaModule.cpp
@@ -1118,28 +1118,36 @@ void FPersonaModule::AddCommonToolbarExtensions(FToolBarBuilder& InToolbarBuilde
 
 TSharedRef< SWidget > FPersonaModule::GenerateCreateAssetMenu(TWeakPtr<IPersonaToolkit> InWeakPersonaToolkit) const
 {
+	IAssetTools& AssetTools = FModuleManager::LoadModuleChecked<FAssetToolsModule>("AssetTools").Get();
+
 	const bool bShouldCloseWindowAfterMenuSelection = true;
 	FMenuBuilder MenuBuilder(bShouldCloseWindowAfterMenuSelection, NULL);
 
 	// Create Animation menu
 	MenuBuilder.BeginSection("CreateAnimation", LOCTEXT("CreateAnimationMenuHeading", "Animation"));
 	{
-		// create menu
-		MenuBuilder.AddSubMenu(
-			LOCTEXT("CreateAnimationSubmenu", "Create Animation"),
-			LOCTEXT("CreateAnimationSubmenu_ToolTip", "Create Animation for this skeleton"),
-			FNewMenuDelegate::CreateRaw(this, &FPersonaModule::FillCreateAnimationMenu, InWeakPersonaToolkit),
-			false,
-			FSlateIcon(FEditorStyle::GetStyleSetName(), "Persona.AssetActions.CreateAnimAsset")
-			);
+		if(AssetTools.IsAssetClassSupported(UAnimSequence::StaticClass()))
+		{
+			// create menu
+			MenuBuilder.AddSubMenu(
+				LOCTEXT("CreateAnimationSubmenu", "Create Animation"),
+				LOCTEXT("CreateAnimationSubmenu_ToolTip", "Create Animation for this skeleton"),
+				FNewMenuDelegate::CreateRaw(this, &FPersonaModule::FillCreateAnimationMenu, InWeakPersonaToolkit),
+				false,
+				FSlateIcon(FEditorStyle::GetStyleSetName(), "Persona.AssetActions.CreateAnimAsset")
+				);
+		}
 
-		MenuBuilder.AddSubMenu(
-			LOCTEXT("CreatePoseAssetSubmenu", "Create PoseAsset"),
-			LOCTEXT("CreatePoseAsssetSubmenu_ToolTip", "Create PoseAsset for this skeleton"),
-			FNewMenuDelegate::CreateRaw(this, &FPersonaModule::FillCreatePoseAssetMenu, InWeakPersonaToolkit),
-			false,
-			FSlateIcon(FEditorStyle::GetStyleSetName(), "ClassIcon.PoseAsset")
-		);
+		if (AssetTools.IsAssetClassSupported(UPoseAsset::StaticClass()))
+		{
+			MenuBuilder.AddSubMenu(
+				LOCTEXT("CreatePoseAssetSubmenu", "Create PoseAsset"),
+				LOCTEXT("CreatePoseAsssetSubmenu_ToolTip", "Create PoseAsset for this skeleton"),
+				FNewMenuDelegate::CreateRaw(this, &FPersonaModule::FillCreatePoseAssetMenu, InWeakPersonaToolkit),
+				false,
+				FSlateIcon(FEditorStyle::GetStyleSetName(), "ClassIcon.PoseAsset")
+			);
+		}
 	}
 	MenuBuilder.EndSection();
 
diff --git a/Engine/Source/Editor/Persona/Private/SAssetFamilyShortcutBar.cpp b/Engine/Source/Editor/Persona/Private/SAssetFamilyShortcutBar.cpp
--- a/Engine/Source/Editor/Persona/Private/SAssetFamilyShortcutBar.cpp
+++ b/Engine/Source/Editor/Persona/Private/SAssetFamilyShortcutBar.cpp
@@ -26,6 +26,8 @@
 #include "Framework/Application/SlateApplication.h"
 #include "Subsystems/AssetEditorSubsystem.h"
 #include "Styling/ToolBarStyle.h"
+#include "AssetToolsModule.h"
+#include "IAssetTools.h"
 
 #define LOCTEXT_NAMESPACE "SAssetFamilyShortcutBar"
 
@@ -520,7 +522,6 @@ void SAssetFamilyShortcutBar::Construct(const FArguments& InArgs, const TSharedR
 	TArray<UClass*> AssetTypes;
 	InAssetFamily->GetAssetTypes(AssetTypes);
 
-	int32 AssetTypeIndex = 0;
 	for (UClass* Class : AssetTypes)
 	{
 		FAssetData AssetData = InAssetFamily->FindAssetOfType(Class);
@@ -529,9 +530,12 @@ void SAssetFamilyShortcutBar::Construct(const FArguments& InArgs, const TSharedR
 		.Padding(0.0f, 4.0f, 16.0f, 4.0f)
 		[
 			SNew(SAssetShortcut, InHostingApp, InAssetFamily, AssetData, ThumbnailPool.ToSharedRef())
+			.Visibility_Lambda([Class]()
+			{
+				IAssetTools& AssetTools = FModuleManager::LoadModuleChecked<FAssetToolsModule>("AssetTools").Get();
+				return AssetTools.IsAssetClassSupported(Class) ? EVisibility::Visible : EVisibility::Collapsed;
+			})
 		];
-
-		AssetTypeIndex++;
 	}
 
 	ChildSlot
diff --git a/Engine/Source/Editor/SkeletalMeshEditor/Private/SkeletalMeshEditor.cpp b/Engine/Source/Editor/SkeletalMeshEditor/Private/SkeletalMeshEditor.cpp
--- a/Engine/Source/Editor/SkeletalMeshEditor/Private/SkeletalMeshEditor.cpp
+++ b/Engine/Source/Editor/SkeletalMeshEditor/Private/SkeletalMeshEditor.cpp
@@ -56,6 +56,9 @@
 #include "Misc/CoreMisc.h"
 #include "Toolkits/AssetEditorToolkitMenuContext.h"
 #include "MeshMergeModule.h"
+#include "AssetToolsModule.h"
+#include "IAssetTools.h"
+#include "SkeletalMeshEditorContextMenuContext.h"
 #include "InterchangeAssetImportData.h"
 
 const FName SkeletalMeshEditorAppIdentifier = FName(TEXT("SkeletalMeshEditorApp"));
@@ -743,112 +746,6 @@ bool FSkeletalMeshEditor::ProcessCommandBindings(const FKeyEvent& InKeyEvent) co
 }
 
 
-void FSkeletalMeshEditor::FillMeshClickMenu(FMenuBuilder& MenuBuilder, HActor* HitProxy, const FViewportClick& Click)
-{
-	UDebugSkelMeshComponent* MeshComp = GetPersonaToolkit()->GetPreviewMeshComponent();
-
-	// Must have hit something, but if the preview is invalid, bail
-	if(!MeshComp)
-	{
-		return;
-	}
-
-	const int32 LodIndex = MeshComp->GetPredictedLODLevel();
-	const int32 SectionIndex = HitProxy->SectionIndex;
-
-	TSharedRef<SWidget> InfoWidget = SNew(SBox)
-		.HAlign(HAlign_Fill)
-		.VAlign(VAlign_Fill)
-		.Padding(FMargin(2.5f, 5.0f, 2.5f, 0.0f))
-		[
-			SNew(SBorder)
-			.HAlign(HAlign_Fill)
-			.VAlign(VAlign_Fill)
-			//.Padding(FMargin(0.0f, 10.0f, 0.0f, 0.0f))
-			.BorderImage(FEditorStyle::GetBrush("ToolPanel.GroupBorder"))
-			[
-				SNew(SBox)
-				.HAlign(HAlign_Center)
-				.VAlign(VAlign_Center)
-				[
-					SNew(STextBlock)
-					.Font(FEditorStyle::GetFontStyle("CurveEd.LabelFont"))
-				.Text(FText::Format(LOCTEXT("MeshClickMenu_SectionInfo", "LOD{0} - Section {1}"), LodIndex, SectionIndex))
-				]
-			]
-		];
-
-
-	MenuBuilder.AddWidget(InfoWidget, FText::GetEmpty(), true, false);
-
-	MenuBuilder.BeginSection(TEXT("MeshClickMenu_Asset"), LOCTEXT("MeshClickMenu_Section_Asset", "Asset"));
-	{
-		FUIAction Action;
-		Action.CanExecuteAction = FCanExecuteAction::CreateSP(this, &FSkeletalMeshEditor::CanApplyClothing, LodIndex, SectionIndex);
-
-		MenuBuilder.AddSubMenu(
-			LOCTEXT("MeshClickMenu_AssetApplyMenu", "Apply Clothing Data..."),
-			LOCTEXT("MeshClickMenu_AssetApplyMenu_ToolTip", "Select clothing data to apply to the selected section."),
-			FNewMenuDelegate::CreateSP(this, &FSkeletalMeshEditor::FillApplyClothingAssetMenu, LodIndex, SectionIndex),
-			Action,
-			TEXT(""),
-			EUserInterfaceActionType::Button
-			);
-
-		Action.ExecuteAction = FExecuteAction::CreateSP(this, &FSkeletalMeshEditor::OnRemoveClothingAssetMenuItemClicked, LodIndex, SectionIndex);
-		Action.CanExecuteAction = FCanExecuteAction::CreateSP(this, &FSkeletalMeshEditor::CanRemoveClothing, LodIndex, SectionIndex);
-
-		MenuBuilder.AddMenuEntry(
-			LOCTEXT("MeshClickMenu_RemoveClothing", "Remove Clothing Data"),
-			LOCTEXT("MeshClickMenu_RemoveClothing_ToolTip", "Remove the currently assigned clothing data."),
-			FSlateIcon(),
-			Action
-			);
-			
-		Action.ExecuteAction = FExecuteAction();
-		Action.CanExecuteAction = FCanExecuteAction::CreateSP(this, &FSkeletalMeshEditor::CanCreateClothing, LodIndex, SectionIndex);
-
-		MenuBuilder.AddSubMenu(
-			LOCTEXT("MeshClickMenu_CreateClothing_Label", "Create Clothing Data from Section"),
-			LOCTEXT("MeshClickMenu_CreateClothing_ToolTip", "Create a new clothing data using the selected section as a simulation mesh"),
-			FNewMenuDelegate::CreateSP(this, &FSkeletalMeshEditor::FillCreateClothingMenu, LodIndex, SectionIndex),
-			Action,
-			TEXT(""),
-			EUserInterfaceActionType::Button
-			);
-
-		Action.CanExecuteAction = FCanExecuteAction::CreateSP(this, &FSkeletalMeshEditor::CanCreateClothingLod, LodIndex, SectionIndex);
-
-		MenuBuilder.AddSubMenu(
-			LOCTEXT("MeshClickMenu_CreateClothingNewLod_Label", "Create Clothing LOD from Section"),
-			LOCTEXT("MeshClickMenu_CreateClothingNewLod_ToolTip", "Create a clothing simulation mesh from the selected section and add it as a LOD to existing clothing data."),
-			FNewMenuDelegate::CreateSP(this, &FSkeletalMeshEditor::FillCreateClothingLodMenu, LodIndex, SectionIndex),
-			Action,
-			TEXT(""),
-			EUserInterfaceActionType::Button
-		);
-
-
-		if (SkeletalMesh != nullptr && SkeletalMesh->GetImportedModel()->LODModels.IsValidIndex(LodIndex))
-		{
-			const FSkeletalMeshLODInfo* SkeletalMeshLODInfo = SkeletalMesh->GetLODInfo(LodIndex);
-			if (SkeletalMeshLODInfo != nullptr)
-			{
-				FUIAction ActionRemoveSection;
-				ActionRemoveSection.ExecuteAction = FExecuteAction::CreateSP(this, &FSkeletalMeshEditor::OnRemoveSectionFromLodAndBelowMenuItemClicked, LodIndex, SectionIndex);
-
-				MenuBuilder.AddMenuEntry(
-					FText::Format(LOCTEXT("MeshClickMenu_RemoveSectionFromLodAndBelow", "Generate section {1} up to LOD {0}"), LodIndex, SectionIndex),
-					FText::Format(LOCTEXT("MeshClickMenu_RemoveSectionFromLodAndBelow_Tooltip", "Generated LODs will use section {1} up to LOD {0}, and ignore it for lower quality LODs"), LodIndex, SectionIndex),
-					FSlateIcon(),
-					ActionRemoveSection
-				);
-			}
-		}
-	}
-	MenuBuilder.EndSection();
-}
-
 void FSkeletalMeshEditor::OnRemoveSectionFromLodAndBelowMenuItemClicked(int32 LodIndex, int32 SectionIndex)
 {
 	if (SkeletalMesh == nullptr || !SkeletalMesh->GetImportedModel()->LODModels.IsValidIndex(LodIndex) || !SkeletalMesh->GetImportedModel()->LODModels[LodIndex].Sections.IsValidIndex(SectionIndex))
@@ -1226,18 +1123,155 @@ void FSkeletalMeshEditor::ExtendMenu()
 	ISkeletalMeshEditorModule& SkeletalMeshEditorModule = FModuleManager::GetModuleChecked<ISkeletalMeshEditorModule>("SkeletalMeshEditor");
 	AddMenuExtender(SkeletalMeshEditorModule.GetMenuExtensibilityManager()->GetAllExtenders(GetToolkitCommands(), GetEditingObjects()));
 
+	FToolMenuOwnerScoped OwnerScoped(this);
+
 	UToolMenu* AssetMenu = UToolMenus::Get()->ExtendMenu("MainFrame.MainMenu.Asset");
-	FToolMenuSection& AssetSection = AssetMenu->FindOrAddSection("AssetEditorActions");
-	const FName SkeletalMeshToolkitName = GetToolkitFName();
-	FToolMenuEntry& Entry = AssetSection.AddDynamicEntry("AssetManagerEditorSkeletalMeshCommands", FNewToolMenuSectionDelegate::CreateLambda([SkeletalMeshToolkitName](FToolMenuSection& InSection)
+	{
+		FToolMenuSection& AssetSection = AssetMenu->FindOrAddSection("AssetEditorActions");
+		const FName SkeletalMeshToolkitName = GetToolkitFName();
+		FToolMenuEntry& Entry = AssetSection.AddDynamicEntry("AssetManagerEditorSkeletalMeshCommands", FNewToolMenuSectionDelegate::CreateLambda([SkeletalMeshToolkitName](FToolMenuSection& InSection)
+			{
+				UAssetEditorToolkitMenuContext* MenuContext = InSection.FindContext<UAssetEditorToolkitMenuContext>();
+				if (MenuContext && MenuContext->Toolkit.IsValid() && MenuContext->Toolkit.Pin()->GetToolkitFName() == SkeletalMeshToolkitName)
+				{
+					InSection.AddMenuEntry(FSkeletalMeshEditorCommands::Get().BakeMaterials);
+				}
+			}
+		));
+	}
+
+	UToolMenu* ViewportMenu = UToolMenus::Get()->RegisterMenu("SkeletalMeshEditor.MeshContextMenu");
+	{
+		FToolMenuSection& AssetSection = ViewportMenu->AddSection("Asset", LOCTEXT("MeshClickMenu_Section_Asset", "Asset"));
+		AssetSection.AddDynamicEntry("DynamicAssetEntries", FNewToolMenuSectionDelegate::CreateLambda([](FToolMenuSection& InSection)
 		{
-			UAssetEditorToolkitMenuContext* MenuContext = InSection.FindContext<UAssetEditorToolkitMenuContext>();
-			if (MenuContext && MenuContext->Toolkit.IsValid() && MenuContext->Toolkit.Pin()->GetToolkitFName() == SkeletalMeshToolkitName)
+			if (USkeletalMeshEditorContextMenuContext* MenuContext = InSection.FindContext<USkeletalMeshEditorContextMenuContext>())
 			{
-				InSection.AddMenuEntry(FSkeletalMeshEditorCommands::Get().BakeMaterials);
+				const int32 LodIndex = MenuContext->LodIndex;
+				const int32 SectionIndex = MenuContext->SectionIndex;
+
+				TSharedRef<SWidget> InfoWidget = SNew(SBox)
+					.HAlign(HAlign_Fill)
+					.VAlign(VAlign_Fill)
+					.Padding(FMargin(2.5f, 5.0f, 2.5f, 0.0f))
+					[
+						SNew(SBorder)
+						.HAlign(HAlign_Fill)
+						.VAlign(VAlign_Fill)
+						//.Padding(FMargin(0.0f, 10.0f, 0.0f, 0.0f))
+						.BorderImage(FEditorStyle::GetBrush("ToolPanel.GroupBorder"))
+						[
+							SNew(SBox)
+							.HAlign(HAlign_Center)
+							.VAlign(VAlign_Center)
+							[
+								SNew(STextBlock)
+								.Font(FEditorStyle::GetFontStyle("CurveEd.LabelFont"))
+								.Text(FText::Format(LOCTEXT("MeshClickMenu_SectionInfo", "LOD{0} - Section {1}"), LodIndex, SectionIndex))
+							]
+						]
+					];
+
+				InSection.AddEntry(FToolMenuEntry::InitWidget("InfoWidget", InfoWidget, FText::GetEmpty(), true, false, true));
 			}
-		}
-	));
+		}));
+
+		FToolMenuSection& ClothSection = ViewportMenu->AddSection("Clothing", LOCTEXT("MeshClickMenu_Section_Clothing", "Clothing"));
+		ClothSection.AddDynamicEntry("DynamicClothingEntries", FNewToolMenuSectionDelegate::CreateLambda([](FToolMenuSection& InSection)
+		{
+			if (USkeletalMeshEditorContextMenuContext* MenuContext = InSection.FindContext<USkeletalMeshEditorContextMenuContext>())
+			{
+				TSharedPtr<FSkeletalMeshEditor> Editor = GetSkeletalMeshEditor(InSection.Context);
+				if(Editor.IsValid())
+				{
+					const int32 LodIndex = MenuContext->LodIndex;
+					const int32 SectionIndex = MenuContext->SectionIndex;
+
+					IAssetTools& AssetTools = FModuleManager::LoadModuleChecked<FAssetToolsModule>("AssetTools").Get();
+					if (AssetTools.GetAssetClassPermissionList(EAssetClassAction::CreateAsset)->PassesFilter(UClothingAssetBase::StaticClass()->GetFName()))
+					{
+						InSection.AddSubMenu(
+							"ApplyClothing",
+							LOCTEXT("MeshClickMenu_AssetApplyMenu", "Apply Clothing Data..."),
+							LOCTEXT("MeshClickMenu_AssetApplyMenu_ToolTip", "Select clothing data to apply to the selected section."),
+							FNewToolMenuChoice(FNewMenuDelegate::CreateSP(Editor.Get(), &FSkeletalMeshEditor::FillApplyClothingAssetMenu, LodIndex, SectionIndex)),
+							FToolUIActionChoice(FUIAction(
+								FExecuteAction(),
+								FCanExecuteAction::CreateSP(Editor.Get(), &FSkeletalMeshEditor::CanApplyClothing, LodIndex, SectionIndex)
+							)),
+							EUserInterfaceActionType::Button
+						);
+
+						InSection.AddMenuEntry(
+							"RemoveClothing",
+							LOCTEXT("MeshClickMenu_RemoveClothing", "Remove Clothing Data"),
+							LOCTEXT("MeshClickMenu_RemoveClothing_ToolTip", "Remove the currently assigned clothing data."),
+							FSlateIcon(),
+							FToolUIActionChoice(FUIAction(
+								FExecuteAction::CreateSP(Editor.Get(), &FSkeletalMeshEditor::OnRemoveClothingAssetMenuItemClicked, LodIndex, SectionIndex),
+								FCanExecuteAction::CreateSP(Editor.Get(), &FSkeletalMeshEditor::CanRemoveClothing, LodIndex, SectionIndex)
+							))
+						);
+
+						InSection.AddSubMenu(
+							"CreateClothing",
+							LOCTEXT("MeshClickMenu_CreateClothing_Label", "Create Clothing Data from Section"),
+							LOCTEXT("MeshClickMenu_CreateClothing_ToolTip", "Create a new clothing data using the selected section as a simulation mesh"),
+							FNewToolMenuChoice(FNewMenuDelegate::CreateSP(Editor.Get(), &FSkeletalMeshEditor::FillCreateClothingMenu, LodIndex, SectionIndex)),
+							FToolUIActionChoice(FUIAction(
+								FExecuteAction(),
+								FCanExecuteAction::CreateSP(Editor.Get(), &FSkeletalMeshEditor::CanCreateClothing, LodIndex, SectionIndex)
+							)),
+							EUserInterfaceActionType::Button
+						);
+
+						InSection.AddSubMenu(
+							"CreateClothingLOD",
+							LOCTEXT("MeshClickMenu_CreateClothingNewLod_Label", "Create Clothing LOD from Section"),
+							LOCTEXT("MeshClickMenu_CreateClothingNewLod_ToolTip", "Create a clothing simulation mesh from the selected section and add it as a LOD to existing clothing data."),
+							FNewToolMenuChoice(FNewMenuDelegate::CreateSP(Editor.Get(), &FSkeletalMeshEditor::FillCreateClothingLodMenu, LodIndex, SectionIndex)),
+							FToolUIActionChoice(FUIAction(
+								FExecuteAction(),
+								FCanExecuteAction::CreateSP(Editor.Get(), &FSkeletalMeshEditor::CanCreateClothingLod, LodIndex, SectionIndex)
+							)),
+							EUserInterfaceActionType::Button
+						);
+					}
+				}
+			}
+		}));
+		
+		FToolMenuSection& LODSection = ViewportMenu->AddSection("LOD", LOCTEXT("MeshClickMenu_Section_LOD", "LOD"));
+		LODSection.AddDynamicEntry("DynamicLODEntries", FNewToolMenuSectionDelegate::CreateLambda([](FToolMenuSection& InSection)
+		{
+			if (USkeletalMeshEditorContextMenuContext* MenuContext = InSection.FindContext<USkeletalMeshEditorContextMenuContext>())
+			{
+				TSharedPtr<FSkeletalMeshEditor> Editor = GetSkeletalMeshEditor(InSection.Context);
+				if (Editor.IsValid())
+				{
+					const int32 LodIndex = MenuContext->LodIndex;
+					const int32 SectionIndex = MenuContext->SectionIndex;
+
+					if (Editor->SkeletalMesh != nullptr && Editor->SkeletalMesh->GetImportedModel()->LODModels.IsValidIndex(LodIndex))
+					{
+						const FSkeletalMeshLODInfo* SkeletalMeshLODInfo = Editor->SkeletalMesh->GetLODInfo(LodIndex);
+						if (SkeletalMeshLODInfo != nullptr)
+						{
+							InSection.AddMenuEntry(
+								"RemoveSectionFromLodAndBelow",
+								FText::Format(LOCTEXT("MeshClickMenu_RemoveSectionFromLodAndBelow", "Generate section {1} up to LOD {0}"), LodIndex, SectionIndex),
+								FText::Format(LOCTEXT("MeshClickMenu_RemoveSectionFromLodAndBelow_Tooltip", "Generated LODs will use section {1} up to LOD {0}, and ignore it for lower quality LODs"), LodIndex, SectionIndex),
+								FSlateIcon(),
+								FToolUIActionChoice(FUIAction(
+									FExecuteAction::CreateSP(Editor.Get(), &FSkeletalMeshEditor::OnRemoveSectionFromLodAndBelowMenuItemClicked, LodIndex, SectionIndex)
+								))
+							);
+						}
+					}
+				}
+			}
+		}));
+	}
 }
 
 void FSkeletalMeshEditor::BakeMaterials()
@@ -1447,14 +1481,14 @@ void FSkeletalMeshEditor::HandleMeshClick(HActor* HitProxy, const FViewportClick
 
 	if(Click.GetKey() == EKeys::RightMouseButton)
 	{
-		FMenuBuilder MenuBuilder(true, nullptr);
-
-		FillMeshClickMenu(MenuBuilder, HitProxy, Click);
+		USkeletalMeshEditorContextMenuContext* MenuContext = NewObject<USkeletalMeshEditorContextMenuContext>();
+		MenuContext->LodIndex = Component->GetPredictedLODLevel();
+		MenuContext->SectionIndex = HitProxy->SectionIndex;
 
 		FSlateApplication::Get().PushMenu(
 			FSlateApplication::Get().GetActiveTopLevelWindow().ToSharedRef(),
 			FWidgetPath(),
-			MenuBuilder.MakeWidget(),
+			UToolMenus::Get()->GenerateWidget("SkeletalMeshEditor.MeshContextMenu", FToolMenuContext(MenuContext)),
 			FSlateApplication::Get().GetCursorPos(),
 			FPopupTransitionEffect(FPopupTransitionEffect::ContextMenu)
 			);
diff --git a/Engine/Source/Editor/SkeletalMeshEditor/Private/SkeletalMeshEditor.h b/Engine/Source/Editor/SkeletalMeshEditor/Private/SkeletalMeshEditor.h
--- a/Engine/Source/Editor/SkeletalMeshEditor/Private/SkeletalMeshEditor.h
+++ b/Engine/Source/Editor/SkeletalMeshEditor/Private/SkeletalMeshEditor.h
@@ -140,7 +140,6 @@ private:
 	void HandleMeshClick(HActor* HitProxy, const FViewportClick& Click);
 
 	// Clothing menu handlers (builds and handles clothing context menu options)
-	void FillMeshClickMenu(FMenuBuilder& MenuBuilder, HActor* HitProxy, const FViewportClick& Click);
 	void FillApplyClothingAssetMenu(FMenuBuilder& MenuBuilder, int32 InLodIndex, int32 InSectionIndex);
 	void FillCreateClothingMenu(FMenuBuilder& MenuBuilder, int32 InLodIndex, int32 InSectionIndex);
 	void FillCreateClothingLodMenu(FMenuBuilder& MenuBuilder, int32 InLodIndex, int32 InSectionIndex);
diff --git a/Engine/Source/Editor/SkeletalMeshEditor/Public/SkeletalMeshEditorContextMenuContext.h b/Engine/Source/Editor/SkeletalMeshEditor/Public/SkeletalMeshEditorContextMenuContext.h
new file mode 100644
--- /dev/null
+++ b/Engine/Source/Editor/SkeletalMeshEditor/Public/SkeletalMeshEditorContextMenuContext.h
@@ -0,0 +1,20 @@
+// Copyright Epic Games, Inc. All Rights Reserved.
+
+#pragma once
+
+#include "CoreMinimal.h"
+#include "ToolMenuContext.h"
+#include "SkeletalMeshEditorContextMenuContext.generated.h"
+
+UCLASS()
+class SKELETALMESHEDITOR_API USkeletalMeshEditorContextMenuContext : public UObject
+{
+	GENERATED_BODY()
+
+public:
+	// The LOD index that was clicked
+	int32 LodIndex = INDEX_NONE;
+
+	// The section index that was clicked
+	int32 SectionIndex = INDEX_NONE;
+};
\ No newline at end of file
diff --git a/Engine/Source/Editor/SkeletonEditor/Private/SSkeletonTree.cpp b/Engine/Source/Editor/SkeletonEditor/Private/SSkeletonTree.cpp
--- a/Engine/Source/Editor/SkeletonEditor/Private/SSkeletonTree.cpp
+++ b/Engine/Source/Editor/SkeletonEditor/Private/SSkeletonTree.cpp
@@ -66,6 +66,8 @@
 #include "IPinnedCommandList.h"
 #include "PersonaModule.h"
 #include "SPositiveActionButton.h"
+#include "ToolMenus.h"
+#include "ToolMenuMisc.h"
 
 #define LOCTEXT_NAMESPACE "SSkeletonTree"
 
@@ -210,6 +212,10 @@ void SSkeletonTree::Construct(const FArguments& InArgs, const TSharedRef<FEditab
 	FSkeletonTreeCommands::Register();
 	BindCommands();
 
+	CreateBlendProfileMenu(nullptr);
+	CreateNewMenu();
+	CreateFilterMenu();
+
 	this->ChildSlot
 	[
 		SNew( SOverlay )
@@ -234,7 +240,7 @@ void SSkeletonTree::Construct(const FArguments& InArgs, const TSharedRef<FEditab
 				.Padding(FMargin(6.f, 0.0))
 				[
 					SNew(SPositiveActionButton)
-					.OnGetMenuContent( this, &SSkeletonTree::CreateNewMenu )
+					.OnGetMenuContent( this, &SSkeletonTree::CreateNewMenuWidget )
 					.Icon(FAppStyle::Get().GetBrush("Icons.Plus"))
 				]
 
@@ -258,7 +264,7 @@ void SSkeletonTree::Construct(const FArguments& InArgs, const TSharedRef<FEditab
 					.ComboButtonStyle(&FAppStyle::Get().GetWidgetStyle<FComboButtonStyle>("SimpleComboButton"))
 					.ForegroundColor(FSlateColor::UseStyle())
 					.ContentPadding(2.0f)
-					.OnGetMenuContent( this, &SSkeletonTree::CreateFilterMenu )
+					.OnGetMenuContent( this, &SSkeletonTree::CreateFilterMenuWidget )
 					.ToolTipText( this, &SSkeletonTree::GetFilterMenuTooltip )
 					.AddMetaData<FTagMetaData>(TEXT("SkelTree.Bones"))
 					.HasDownArrow(true)
@@ -343,31 +349,36 @@ void SSkeletonTree::BindCommands()
 		MenuActions.ShowAllBones,
 		FExecuteAction::CreateSP( this, &SSkeletonTree::SetBoneFilter, EBoneFilter::All ),
 		FCanExecuteAction(),
-		FIsActionChecked::CreateSP( this, &SSkeletonTree::IsBoneFilter, EBoneFilter::All ));
+		FIsActionChecked::CreateSP( this, &SSkeletonTree::IsBoneFilter, EBoneFilter::All ),
+		FIsActionButtonVisible::CreateSP( Builder.Get(), &ISkeletonTreeBuilder::IsShowingBones ));
 
 	CommandList.MapAction(
 		MenuActions.ShowMeshBones,
 		FExecuteAction::CreateSP( this, &SSkeletonTree::SetBoneFilter, EBoneFilter::Mesh ),
 		FCanExecuteAction(),
-		FIsActionChecked::CreateSP( this, &SSkeletonTree::IsBoneFilter, EBoneFilter::Mesh ));
+		FIsActionChecked::CreateSP( this, &SSkeletonTree::IsBoneFilter, EBoneFilter::Mesh ),
+		FIsActionButtonVisible::CreateSP(Builder.Get(), &ISkeletonTreeBuilder::IsShowingBones));
 
 	CommandList.MapAction(
 		MenuActions.ShowLODBones,
 		FExecuteAction::CreateSP(this, &SSkeletonTree::SetBoneFilter, EBoneFilter::LOD),
 		FCanExecuteAction(),
-		FIsActionChecked::CreateSP(this, &SSkeletonTree::IsBoneFilter, EBoneFilter::LOD));
+		FIsActionChecked::CreateSP(this, &SSkeletonTree::IsBoneFilter, EBoneFilter::LOD),
+		FIsActionButtonVisible::CreateSP(Builder.Get(), &ISkeletonTreeBuilder::IsShowingBones));
 	
 	CommandList.MapAction(
 		MenuActions.ShowWeightedBones,
 		FExecuteAction::CreateSP( this, &SSkeletonTree::SetBoneFilter, EBoneFilter::Weighted ),
 		FCanExecuteAction(),
-		FIsActionChecked::CreateSP( this, &SSkeletonTree::IsBoneFilter, EBoneFilter::Weighted ));
+		FIsActionChecked::CreateSP( this, &SSkeletonTree::IsBoneFilter, EBoneFilter::Weighted ),
+		FIsActionButtonVisible::CreateSP(Builder.Get(), &ISkeletonTreeBuilder::IsShowingBones));
 
 	CommandList.MapAction(
 		MenuActions.HideBones,
 		FExecuteAction::CreateSP( this, &SSkeletonTree::SetBoneFilter, EBoneFilter::None ),
 		FCanExecuteAction(),
-		FIsActionChecked::CreateSP( this, &SSkeletonTree::IsBoneFilter, EBoneFilter::None ));
+		FIsActionChecked::CreateSP( this, &SSkeletonTree::IsBoneFilter, EBoneFilter::None ),
+		FIsActionButtonVisible::CreateSP(Builder.Get(), &ISkeletonTreeBuilder::IsShowingBones));
 
 	CommandList.EndGroup();
 
@@ -378,31 +389,36 @@ void SSkeletonTree::BindCommands()
 		MenuActions.ShowActiveSockets,
 		FExecuteAction::CreateSP( this, &SSkeletonTree::SetSocketFilter, ESocketFilter::Active ),
 		FCanExecuteAction(),
-		FIsActionChecked::CreateSP( this, &SSkeletonTree::IsSocketFilter, ESocketFilter::Active ));
+		FIsActionChecked::CreateSP( this, &SSkeletonTree::IsSocketFilter, ESocketFilter::Active ),
+		FIsActionButtonVisible::CreateSP(Builder.Get(), &ISkeletonTreeBuilder::IsShowingSockets ));
 
 	CommandList.MapAction(
 		MenuActions.ShowMeshSockets,
 		FExecuteAction::CreateSP( this, &SSkeletonTree::SetSocketFilter, ESocketFilter::Mesh ),
 		FCanExecuteAction(),
-		FIsActionChecked::CreateSP( this, &SSkeletonTree::IsSocketFilter, ESocketFilter::Mesh ));
+		FIsActionChecked::CreateSP( this, &SSkeletonTree::IsSocketFilter, ESocketFilter::Mesh ),
+		FIsActionButtonVisible::CreateSP(Builder.Get(), &ISkeletonTreeBuilder::IsShowingSockets ));
 
 	CommandList.MapAction(
 		MenuActions.ShowSkeletonSockets,
 		FExecuteAction::CreateSP( this, &SSkeletonTree::SetSocketFilter, ESocketFilter::Skeleton ),
 		FCanExecuteAction(),
-		FIsActionChecked::CreateSP( this, &SSkeletonTree::IsSocketFilter, ESocketFilter::Skeleton ));
+		FIsActionChecked::CreateSP( this, &SSkeletonTree::IsSocketFilter, ESocketFilter::Skeleton ),
+		FIsActionButtonVisible::CreateSP(Builder.Get(), &ISkeletonTreeBuilder::IsShowingSockets ));
 
 	CommandList.MapAction(
 		MenuActions.ShowAllSockets,
 		FExecuteAction::CreateSP( this, &SSkeletonTree::SetSocketFilter, ESocketFilter::All ),
 		FCanExecuteAction(),
-		FIsActionChecked::CreateSP( this, &SSkeletonTree::IsSocketFilter, ESocketFilter::All ));
+		FIsActionChecked::CreateSP( this, &SSkeletonTree::IsSocketFilter, ESocketFilter::All ),
+		FIsActionButtonVisible::CreateSP(Builder.Get(), &ISkeletonTreeBuilder::IsShowingSockets ));
 
 	CommandList.MapAction(
 		MenuActions.HideSockets,
 		FExecuteAction::CreateSP( this, &SSkeletonTree::SetSocketFilter, ESocketFilter::None ),
 		FCanExecuteAction(),
-		FIsActionChecked::CreateSP( this, &SSkeletonTree::IsSocketFilter, ESocketFilter::None ));
+		FIsActionChecked::CreateSP( this, &SSkeletonTree::IsSocketFilter, ESocketFilter::None ),
+		FIsActionButtonVisible::CreateSP(Builder.Get(), &ISkeletonTreeBuilder::IsShowingSockets ));
 
 	CommandList.EndGroup();
 
@@ -410,13 +426,15 @@ void SSkeletonTree::BindCommands()
 		MenuActions.ShowRetargeting,
 		FExecuteAction::CreateSP(this, &SSkeletonTree::OnChangeShowingAdvancedOptions),
 		FCanExecuteAction(),
-		FIsActionChecked::CreateSP(this, &SSkeletonTree::IsShowingAdvancedOptions));
+		FIsActionChecked::CreateSP(this, &SSkeletonTree::IsShowingAdvancedOptions),
+		FIsActionButtonVisible::CreateLambda([this]() { return Builder->IsShowingBones() && bAllowSkeletonOperations; }));
 
 	CommandList.MapAction(
 		MenuActions.ShowDebugVisualization,
 		FExecuteAction::CreateSP(this, &SSkeletonTree::OnChangeShowingDebugVisualizationOptions),
 		FCanExecuteAction(),
-		FIsActionChecked::CreateSP(this, &SSkeletonTree::IsShowingDebugVisualizationOptions));
+		FIsActionChecked::CreateSP(this, &SSkeletonTree::IsShowingDebugVisualizationOptions),
+		FIsActionButtonVisible::CreateLambda([this](){ return bShowDebugVisualizationOptions; }));
 
 	// Socket manipulation commands
 	CommandList.MapAction(
@@ -743,7 +761,7 @@ TSharedPtr< SWidget > SSkeletonTree::CreateContextMenu()
 			{
 				MenuBuilder.AddSubMenu(LOCTEXT("AddVirtualBone", "Add Virtual Bone"),
 					LOCTEXT("AddVirtualBone_ToolTip", "Adds a virtual bone to the skeleton."),
-					FNewMenuDelegate::CreateSP(this, &SSkeletonTree::FillVirtualBoneSubmenu, SelectedItems));
+					FNewMenuDelegate::CreateSP(this, &SSkeletonTree::FillVirtualBoneSubmenu));
 			}
 
 			MenuBuilder.EndSection();
@@ -944,13 +962,15 @@ bool GetSourceNameFromItem(TSharedPtr<ISkeletonTreeItem> SourceBone, FName& OutN
 	return false;
 }
 
-void SSkeletonTree::FillVirtualBoneSubmenu(FMenuBuilder& MenuBuilder, TArray<TSharedPtr<ISkeletonTreeItem>> SourceBones)
+void SSkeletonTree::FillVirtualBoneSubmenu(FMenuBuilder& MenuBuilder)
 {
+	TArray<TSharedPtr<ISkeletonTreeItem>> SelectedItems = SkeletonTreeView->GetSelectedItems();
+
 	const bool bShowVirtualBones = false;
 	TSharedRef<SBoneTreeMenu> MenuContent = SNew(SBoneTreeMenu)
 	.bShowVirtualBones(false)
 	.Title(LOCTEXT("TargetBonePickerTitle", "Pick Target Bone..."))
-	.OnBoneSelectionChanged(this, &SSkeletonTree::OnVirtualTargetBonePicked, SourceBones)
+	.OnBoneSelectionChanged(this, &SSkeletonTree::OnVirtualTargetBonePicked, SelectedItems)
 	.OnGetReferenceSkeleton(this, &SSkeletonTree::OnGetReferenceSkeleton);
 	MenuBuilder.AddWidget(MenuContent, FText::GetEmpty(), true);
 
@@ -1580,11 +1600,7 @@ void SSkeletonTree::HandlePackageReloaded(const EPackageReloadPhase InPackageRel
 
 TSharedRef<SWidget> SSkeletonTree::GetBlendProfileColumnMenuContent()
 {
-	const FSkeletonTreeCommands& Actions = FSkeletonTreeCommands::Get();
-
-	FMenuBuilder MenuBuilder( true /*CloseAfterSelection*/, UICommandList );
-	GetBlendProfileMenu(MenuBuilder);
-	return MenuBuilder.MakeWidget();
+	return UToolMenus::Get()->GenerateWidget("SkeletonTree.BlendProfilesMenu", FToolMenuContext(UICommandList, Extenders));
 }
 
 void SSkeletonTree::ExpandTreeOnSelection(TSharedPtr<ISkeletonTreeItem> RowToExpand, bool bForce)
@@ -1600,47 +1616,61 @@ void SSkeletonTree::ExpandTreeOnSelection(TSharedPtr<ISkeletonTreeItem> RowToExp
 	}
 }
 
-void SSkeletonTree::GetBlendProfileMenu(FMenuBuilder& MenuBuilder)
+void SSkeletonTree::CreateBlendProfileMenu(UToolMenu* InMenu)
 {
-	const FSkeletonTreeCommands& Actions = FSkeletonTreeCommands::Get();
+	FToolMenuOwnerScoped OwnerScoped(this);
 
-	MenuBuilder.BeginSection("BlendProfileActions", LOCTEXT("BlendProfiles", "Blend Profiles"));
-	{
-		for (UBlendProfile* Profile : GetEditableSkeletonInternal()->GetBlendProfiles())
-		{
-			MenuBuilder.AddMenuEntry(
-				FText::FromName(Profile->GetFName()),
-				LOCTEXT("SelectBlendProfileTooltip", "Select this profile for editing."),
-				FSlateIcon(),	
-				FUIAction(
-					FExecuteAction::CreateSP(BlendProfilePicker.ToSharedRef(), &SBlendProfilePicker::SetSelectedProfile, Profile, true),
-					FCanExecuteAction(),
-					FIsActionChecked::CreateSP(this, &SSkeletonTree::IsBlendProfileSelected, Profile->GetFName())
-				),
-				NAME_None,
-				EUserInterfaceActionType::RadioButton
-			);
-		}
+	UToolMenu* Menu = InMenu ? InMenu : UToolMenus::Get()->ExtendMenu("SkeletonTree.BlendProfilesMenu");
 
-		MenuBuilder.AddMenuEntry(
-		LOCTEXT("Clear", "Clear"),
-		LOCTEXT("Clear_ToolTip", "Clear the selected blend profile."),
-		FSlateIcon(),
-		FUIAction(FExecuteAction::CreateSP(BlendProfilePicker.ToSharedRef(), &SBlendProfilePicker::OnClearSelection)));
-
-		MenuBuilder.AddSeparator();
-		MenuBuilder.AddMenuEntry(Actions.CreateTimeBlendProfile);
-		MenuBuilder.AddMenuEntry(Actions.CreateWeightBlendProfile);
-		MenuBuilder.AddMenuEntry(Actions.CreateBlendMask);
-		if (BlendProfilePicker->GetSelectedBlendProfileName() != NAME_None)
+	Menu->AddDynamicSection(NAME_None,
+		FNewSectionConstructChoice(FNewToolMenuDelegate::CreateLambda([this](UToolMenu* InMenu)
 		{
-			MenuBuilder.AddSeparator();
-			MenuBuilder.AddMenuEntry(Actions.DeleteCurrentBlendProfile, NAME_None, FText::Format(LOCTEXT("DeleteBlendProfileLabel", "Remove {0}"), FText::FromName(BlendProfilePicker->GetSelectedBlendProfileName())));
-		}
+			const FSkeletonTreeCommands& Actions = FSkeletonTreeCommands::Get();
 
+			FToolMenuSection& ProfilesSection = InMenu->AddSection(TEXT("BlendProfileActions"), LOCTEXT("BlendProfiles", "Blend Profiles"));
 
-	}
-	MenuBuilder.EndSection();
+			for (UBlendProfile* Profile : GetEditableSkeletonInternal()->GetBlendProfiles())
+			{
+				ProfilesSection.AddMenuEntry(
+					Profile->GetFName(),
+					FText::FromName(Profile->GetFName()),
+					LOCTEXT("SelectBlendProfileTooltip", "Select this profile for editing."),
+					FSlateIcon(),
+					FToolUIActionChoice(
+						FUIAction(
+							FExecuteAction::CreateSP(BlendProfilePicker.ToSharedRef(), &SBlendProfilePicker::SetSelectedProfile, Profile, true),
+							FCanExecuteAction(),
+							FIsActionChecked::CreateSP(this, &SSkeletonTree::IsBlendProfileSelected, Profile->GetFName())
+						)
+					),
+					EUserInterfaceActionType::RadioButton
+				);
+			}
+
+			if (BlendProfilePicker->GetSelectedBlendProfileName() != NAME_None)
+			{
+				FToolMenuSection& EditSection = InMenu->AddSection(TEXT("BlendProfileEdit"), LOCTEXT("EditBlendProfilesSection", "Edit"));
+
+				EditSection.AddMenuEntry(
+					"ClearBlendProfile",
+					LOCTEXT("Clear", "Clear Selected"),
+					LOCTEXT("Clear_ToolTip", "Clear the selected blend profile."),
+					FSlateIcon(),
+					FToolUIActionChoice(FUIAction(FExecuteAction::CreateSP(BlendProfilePicker.ToSharedRef(), &SBlendProfilePicker::OnClearSelection))));
+
+				EditSection.AddMenuEntry(
+					Actions.DeleteCurrentBlendProfile, 
+					FText::Format(LOCTEXT("DeleteBlendProfileLabel", "Delete {0}"),
+					FText::FromName(BlendProfilePicker->GetSelectedBlendProfileName())));
+			}
+
+			{
+				FToolMenuSection& NewSection = InMenu->AddSection(TEXT("BlendProfileNew"), LOCTEXT("BlendProfiles", "New"));
+				NewSection.AddMenuEntry(Actions.CreateTimeBlendProfile);
+				NewSection.AddMenuEntry(Actions.CreateWeightBlendProfile);
+				NewSection.AddMenuEntry(Actions.CreateBlendMask);
+			}
+		})));
 } 
 
 void SSkeletonTree::OnCreateBlendProfile(const EBlendProfileMode InMode)
@@ -1748,86 +1778,85 @@ void SSkeletonTree::NotifyUser( FNotificationInfo& NotificationInfo )
 	}
 }
 
-TSharedRef< SWidget > SSkeletonTree::CreateNewMenu()
+void SSkeletonTree::CreateNewMenu()
 {
+	FToolMenuOwnerScoped OwnerScoped(this);
+
 	const FSkeletonTreeCommands& Actions = FSkeletonTreeCommands::Get();
-	TArray<TSharedPtr<ISkeletonTreeItem>> SelectedItems = SkeletonTreeView->GetSelectedItems();
 
-	FMenuBuilder MenuBuilder( true /*CloseAfterSelection*/ , UICommandList, Extenders );
+	UToolMenu* Menu = UToolMenus::Get()->ExtendMenu("SkeletonTree.NewMenu");
 
-	MenuBuilder.BeginSection("CreateNew", LOCTEXT("SkeletonCreateNew", "Create"));
 	{
-		MenuBuilder.AddMenuEntry(Actions.AddSocket);
+		FToolMenuSection& CreateSection = Menu->AddSection("CreateNew", LOCTEXT("SkeletonCreateNew", "Create"));
 
-		MenuBuilder.AddSubMenu(LOCTEXT("AddVirtualBone", "Add Virtual Bone"),
+		CreateSection.AddMenuEntry(Actions.AddSocket);
+		CreateSection.AddSubMenu(
+			"VirtualBones",
+			LOCTEXT("AddVirtualBone", "Add Virtual Bone"),
 			LOCTEXT("AddVirtualBone_ToolTip", "Adds a virtual bone to the skeleton."),
-			FNewMenuDelegate::CreateSP(this, &SSkeletonTree::FillVirtualBoneSubmenu, SelectedItems));
+			FNewToolMenuChoice(FNewMenuDelegate::CreateSP(this, &SSkeletonTree::FillVirtualBoneSubmenu)));
 	}
-	MenuBuilder.EndSection();
 
-	MenuBuilder.BeginSection("Blend");
 	{
-		MenuBuilder.AddMenuEntry(Actions.CreateTimeBlendProfile);
-		MenuBuilder.AddMenuEntry(Actions.CreateWeightBlendProfile);
-		MenuBuilder.AddMenuEntry(Actions.CreateBlendMask);
+		FToolMenuSection& BlendSection = Menu->AddSection("Blend", LOCTEXT("SkeletonBlend", "Blend"));
+		BlendSection.AddMenuEntry(Actions.CreateTimeBlendProfile);
+		BlendSection.AddMenuEntry(Actions.CreateWeightBlendProfile);
+		BlendSection.AddMenuEntry(Actions.CreateBlendMask);
 	}
-	MenuBuilder.EndSection();
+}
 
-	return MenuBuilder.MakeWidget();
+TSharedRef< SWidget > SSkeletonTree::CreateNewMenuWidget()
+{
+	return UToolMenus::Get()->GenerateWidget("SkeletonTree.NewMenu", FToolMenuContext(UICommandList, Extenders)); 
 }
 
-TSharedRef< SWidget > SSkeletonTree::CreateFilterMenu()
+void SSkeletonTree::CreateFilterMenu()
 {
+	FToolMenuOwnerScoped OwnerScoped(this);
+
 	const FSkeletonTreeCommands& Actions = FSkeletonTreeCommands::Get();
 
-	const bool CloseAfterSelection = true;
-	FMenuBuilder MenuBuilder( CloseAfterSelection, UICommandList, Extenders );
+	UToolMenu* Menu = UToolMenus::Get()->ExtendMenu("SkeletonTree.FilterMenu");
 
-	GetBlendProfileMenu(MenuBuilder);
+	{
+		FToolMenuSection& BlendProfilesSection = Menu->AddSection("BlendProfiles", LOCTEXT("BlendProfilesMenuHeading", "Blend Profiles"));
+		BlendProfilesSection.AddSubMenu(
+			"BlendProfiles", 
+			LOCTEXT("BlendProfilesSubMenu", "Blend Profiles"),
+			LOCTEXT("BlendProfilesSubMenuTooltip", "Edit Blend Profiles in this Skeleton"), 
+			FNewToolMenuChoice(FNewToolMenuDelegate::CreateSP(this, &SSkeletonTree::CreateBlendProfileMenu)));
+	}
 
-	MenuBuilder.BeginSection("FilterOptions", LOCTEXT("OptionsMenuHeading", "Options"));
 	{
-		if(Builder->IsShowingBones() && bAllowSkeletonOperations)
-		{
-			MenuBuilder.AddMenuEntry(Actions.ShowRetargeting);
-		}
-		MenuBuilder.AddMenuEntry(Actions.FilteringFlattensHierarchy);
-		MenuBuilder.AddMenuEntry(Actions.HideParentsWhenFiltering);
-		
-		if (bShowDebugVisualizationOptions)
-		{
-			MenuBuilder.AddMenuEntry(Actions.ShowDebugVisualization);
-		}
+		FToolMenuSection& OptionsSection = Menu->AddSection("FilterOptions", LOCTEXT("OptionsMenuHeading", "Options"));
+		OptionsSection.AddMenuEntry(Actions.ShowRetargeting);
+		OptionsSection.AddMenuEntry(Actions.FilteringFlattensHierarchy);
+		OptionsSection.AddMenuEntry(Actions.HideParentsWhenFiltering);
+		OptionsSection.AddMenuEntry(Actions.ShowDebugVisualization);
 	}
-	MenuBuilder.EndSection();
 
-	if(Builder->IsShowingBones())
 	{
-		MenuBuilder.BeginSection("FilterBones", LOCTEXT( "BonesMenuHeading", "Bones" ) );
-		{
-			MenuBuilder.AddMenuEntry( Actions.ShowAllBones );
-			MenuBuilder.AddMenuEntry( Actions.ShowMeshBones );
-			MenuBuilder.AddMenuEntry( Actions.ShowLODBones);
-			MenuBuilder.AddMenuEntry( Actions.ShowWeightedBones );
-			MenuBuilder.AddMenuEntry( Actions.HideBones );
-		}
-		MenuBuilder.EndSection();
+		FToolMenuSection& BonesSection = Menu->AddSection("FilterBones", LOCTEXT("BonesMenuHeading", "Bones"));
+		BonesSection.AddMenuEntry(Actions.ShowAllBones);
+		BonesSection.AddMenuEntry(Actions.ShowMeshBones);
+		BonesSection.AddMenuEntry(Actions.ShowLODBones);
+		BonesSection.AddMenuEntry(Actions.ShowWeightedBones);
+		BonesSection.AddMenuEntry(Actions.HideBones);
 	}
 
-	if(Builder->IsShowingSockets())
 	{
-		MenuBuilder.BeginSection("FilterSockets", LOCTEXT("SocketsMenuHeading", "Sockets"));
-		{
-			MenuBuilder.AddMenuEntry(Actions.ShowActiveSockets);
-			MenuBuilder.AddMenuEntry(Actions.ShowMeshSockets);
-			MenuBuilder.AddMenuEntry(Actions.ShowSkeletonSockets);
-			MenuBuilder.AddMenuEntry(Actions.ShowAllSockets);
-			MenuBuilder.AddMenuEntry(Actions.HideSockets);
-		}
-		MenuBuilder.EndSection();
+		FToolMenuSection& BonesSection = Menu->AddSection("FilterSockets", LOCTEXT("SocketsMenuHeading", "Sockets"));
+		BonesSection.AddMenuEntry(Actions.ShowActiveSockets);
+		BonesSection.AddMenuEntry(Actions.ShowMeshSockets);
+		BonesSection.AddMenuEntry(Actions.ShowSkeletonSockets);
+		BonesSection.AddMenuEntry(Actions.ShowAllSockets);
+		BonesSection.AddMenuEntry(Actions.HideSockets);
 	}
+}
 
-	return MenuBuilder.MakeWidget();
+TSharedRef< SWidget > SSkeletonTree::CreateFilterMenuWidget()
+{
+	return UToolMenus::Get()->GenerateWidget("SkeletonTree.FilterMenu", FToolMenuContext(UICommandList, Extenders));
 }
 
 void SSkeletonTree::SetBoneFilter( EBoneFilter InBoneFilter )
diff --git a/Engine/Source/Editor/SkeletonEditor/Private/SSkeletonTree.h b/Engine/Source/Editor/SkeletonEditor/Private/SSkeletonTree.h
--- a/Engine/Source/Editor/SkeletonEditor/Private/SSkeletonTree.h
+++ b/Engine/Source/Editor/SkeletonEditor/Private/SSkeletonTree.h
@@ -168,10 +168,16 @@ private:
 	TSharedPtr< SWidget > CreateContextMenu();
 
 	/** Called to display the add new menu */
-	TSharedRef< SWidget > CreateNewMenu();
+	TSharedRef< SWidget > CreateNewMenuWidget();
+
+	/** Called to create the add new menu */
+	void CreateNewMenu();
 
 	/** Called to display the filter menu */
-	TSharedRef< SWidget > CreateFilterMenu();
+	TSharedRef< SWidget > CreateFilterMenuWidget();
+
+	/** Called to create the filter menu */
+	void CreateFilterMenu();
 
 	/** Function to copy selected bone name to the clipboard */
 	void OnCopyBoneNames();
@@ -210,7 +216,7 @@ private:
 	void OnPromoteSocket();
 
 	/** Create sub menu to allow users to pick a target bone for the new space switching bone(s) */
-	void FillVirtualBoneSubmenu(FMenuBuilder& MenuBuilder, TArray<TSharedPtr<class ISkeletonTreeItem>> SourceBones);
+	void FillVirtualBoneSubmenu(FMenuBuilder& MenuBuilder);
 
 	/** Handler for user picking a target bone */
 	void OnVirtualTargetBonePicked(FName TargetBoneName, TArray<TSharedPtr<class ISkeletonTreeItem>> SourceBones);
@@ -333,7 +339,7 @@ private:
 	void OnDeleteCurrentBlendProfile();
 
 	/** Generate Blend Profile Menu */
-	void GetBlendProfileMenu(FMenuBuilder& MenuBuilder);
+	void CreateBlendProfileMenu(UToolMenu* InMenu = nullptr);
 
 	TSharedRef<SWidget> GetBlendProfileColumnMenuContent();
 
diff --git a/Engine/Source/Editor/SkeletonEditor/Private/SkeletonEditor.cpp b/Engine/Source/Editor/SkeletonEditor/Private/SkeletonEditor.cpp
--- a/Engine/Source/Editor/SkeletonEditor/Private/SkeletonEditor.cpp
+++ b/Engine/Source/Editor/SkeletonEditor/Private/SkeletonEditor.cpp
@@ -23,6 +23,9 @@
 #include "ISkeletonTreeItem.h"
 #include "Algo/Transform.h"
 #include "PersonaToolMenuContext.h"
+#include "ToolMenus.h"
+#include "ToolMenuMisc.h"
+#include "SkeletonToolMenuContext.h"
 
 const FName SkeletonEditorAppIdentifier = FName(TEXT("SkeletonEditorApp"));
 
@@ -143,10 +146,13 @@ void FSkeletonEditor::InitToolMenuContext(FToolMenuContext& MenuContext)
 {
 	FAssetEditorToolkit::InitToolMenuContext(MenuContext);
 
-	UPersonaToolMenuContext* Context = NewObject<UPersonaToolMenuContext>();
-	Context->SetToolkit(GetPersonaToolkit());
+	UPersonaToolMenuContext* PersonaContext = NewObject<UPersonaToolMenuContext>();
+	PersonaContext->SetToolkit(GetPersonaToolkit());
+	MenuContext.AddObject(PersonaContext);
 
-	MenuContext.AddObject(Context);
+	USkeletonToolMenuContext* SkeletonContext = NewObject<USkeletonToolMenuContext>();
+	SkeletonContext->SkeletonEditor = SharedThis(this);
+	MenuContext.AddObject(SkeletonContext);
 }
 
 void FSkeletonEditor::BindCommands()
@@ -176,8 +182,49 @@ void FSkeletonEditor::BindCommands()
 		FExecuteAction::CreateRaw(&GetPersonaToolkit()->GetPreviewScene().Get(), &IPersonaPreviewScene::TogglePlayback));
 }
 
+TSharedPtr<FSkeletonEditor> FSkeletonEditor::GetSkeletonEditor(const FToolMenuContext& InMenuContext)
+{
+	if (USkeletonToolMenuContext* Context = InMenuContext.FindContext<USkeletonToolMenuContext>())
+	{
+		if (Context->SkeletonEditor.IsValid())
+		{
+			return StaticCastSharedPtr<FSkeletonEditor>(Context->SkeletonEditor.Pin());
+		}
+	}
+
+	return TSharedPtr<FSkeletonEditor>();
+}
+
 void FSkeletonEditor::ExtendToolbar()
 {
+	FToolMenuOwnerScoped OwnerScoped(this);
+
+	// Add in Editor Specific functionality
+	FName ParentName;
+	static const FName MenuName = GetToolMenuToolbarName(ParentName);
+
+	UToolMenu* ToolMenu = UToolMenus::Get()->ExtendMenu(MenuName);
+	const FToolMenuInsert SectionInsertLocation("Asset", EToolMenuInsertType::After);
+
+	{
+		ToolMenu->AddDynamicSection("Persona", FNewToolBarDelegateLegacy::CreateLambda([](FToolBarBuilder& ToolbarBuilder, UToolMenu* InMenu)
+		{
+			TSharedPtr<FSkeletonEditor> SkeletonEditor = GetSkeletonEditor(InMenu->Context);
+			if (SkeletonEditor.IsValid() && SkeletonEditor->PersonaToolkit.IsValid())
+			{
+				FPersonaModule& PersonaModule = FModuleManager::LoadModuleChecked<FPersonaModule>("Persona");
+				PersonaModule.AddCommonToolbarExtensions(ToolbarBuilder, SkeletonEditor->PersonaToolkit.ToSharedRef());
+			}
+		}), SectionInsertLocation);
+	}
+
+	{
+		FToolMenuSection& SkeletonSection = ToolMenu->AddSection("Skeleton", LOCTEXT("ToolbarSkeletonSectionLabel", "Skeleton"), SectionInsertLocation);
+		SkeletonSection.AddEntry(FToolMenuEntry::InitToolBarButton(FSkeletonEditorCommands::Get().AnimNotifyWindow));
+		SkeletonSection.AddEntry(FToolMenuEntry::InitToolBarButton(FSkeletonEditorCommands::Get().RetargetManager, LOCTEXT("Toolbar_RetargetManager", "Retarget Manager")));
+		SkeletonSection.AddEntry(FToolMenuEntry::InitToolBarButton(FSkeletonEditorCommands::Get().ImportMesh));
+	}
+
 	// If the ToolbarExtender is valid, remove it before rebuilding it
 	if (ToolbarExtender.IsValid())
 	{
@@ -209,16 +256,6 @@ void FSkeletonEditor::ExtendToolbar()
 		FToolBarExtensionDelegate::CreateLambda([this](FToolBarBuilder& ToolbarBuilder)
 		{
 			FPersonaModule& PersonaModule = FModuleManager::LoadModuleChecked<FPersonaModule>("Persona");
-			PersonaModule.AddCommonToolbarExtensions(ToolbarBuilder, PersonaToolkit.ToSharedRef());
-
-			ToolbarBuilder.BeginSection("Skeleton");
-			{
-				ToolbarBuilder.AddToolBarButton(FSkeletonEditorCommands::Get().AnimNotifyWindow);
-				ToolbarBuilder.AddToolBarButton(FSkeletonEditorCommands::Get().RetargetManager, NAME_None, LOCTEXT("Toolbar_RetargetSources", "Retarget Sources"));
-				ToolbarBuilder.AddToolBarButton(FSkeletonEditorCommands::Get().ImportMesh);
-			}
-			ToolbarBuilder.EndSection();
-
 			TSharedRef<class IAssetFamily> AssetFamily = PersonaModule.CreatePersonaAssetFamily(Skeleton);
 			AddToolbarWidget(PersonaModule.CreateAssetFamilyShortcutWidget(SharedThis(this), AssetFamily));
 		}	
diff --git a/Engine/Source/Editor/SkeletonEditor/Private/SkeletonEditor.h b/Engine/Source/Editor/SkeletonEditor/Private/SkeletonEditor.h
--- a/Engine/Source/Editor/SkeletonEditor/Private/SkeletonEditor.h
+++ b/Engine/Source/Editor/SkeletonEditor/Private/SkeletonEditor.h
@@ -118,6 +118,7 @@ private:
 
 	void OnImportAsset();
 
+	static TSharedPtr<FSkeletonEditor> GetSkeletonEditor(const FToolMenuContext& InMenuContext);
 public:
 	/** Multicast delegate fired on global undo/redo */
 	FSimpleMulticastDelegate OnPostUndo;
diff --git a/Engine/Source/Editor/SkeletonEditor/Public/ISkeletonTreeBuilder.h b/Engine/Source/Editor/SkeletonEditor/Public/ISkeletonTreeBuilder.h
--- a/Engine/Source/Editor/SkeletonEditor/Public/ISkeletonTreeBuilder.h
+++ b/Engine/Source/Editor/SkeletonEditor/Public/ISkeletonTreeBuilder.h
@@ -102,7 +102,7 @@ struct FSkeletonTreeFilterArgs
 DECLARE_DELEGATE_RetVal_TwoParams(ESkeletonTreeFilterResult, FOnFilterSkeletonTreeItem, const FSkeletonTreeFilterArgs& /*InArgs*/, const TSharedPtr<class ISkeletonTreeItem>& /*InItem*/);
 
 /** Interface to implement to provide custom build logic to skeleton trees */
-class ISkeletonTreeBuilder
+class ISkeletonTreeBuilder : public TSharedFromThis<ISkeletonTreeBuilder>
 {
 public:
 
diff --git a/Engine/Source/Editor/SkeletonEditor/Public/SkeletonToolMenuContext.h b/Engine/Source/Editor/SkeletonEditor/Public/SkeletonToolMenuContext.h
new file mode 100644
--- /dev/null
+++ b/Engine/Source/Editor/SkeletonEditor/Public/SkeletonToolMenuContext.h
@@ -0,0 +1,18 @@
+// Copyright Epic Games, Inc. All Rights Reserved.
+
+#pragma once
+
+#include "CoreMinimal.h"
+#include "ToolMenuContext.h"
+#include "SkeletonToolMenuContext.generated.h"
+
+class ISkeletonEditor;
+
+UCLASS()
+class SKELETONEDITOR_API USkeletonToolMenuContext : public UObject
+{
+	GENERATED_BODY()
+
+public:
+	TWeakPtr<ISkeletonEditor> SkeletonEditor;
+};
\ No newline at end of file
diff --git a/Engine/Source/Editor/UnrealEd/Private/AnimationEditorUtils.cpp b/Engine/Source/Editor/UnrealEd/Private/AnimationEditorUtils.cpp
--- a/Engine/Source/Editor/UnrealEd/Private/AnimationEditorUtils.cpp
+++ b/Engine/Source/Editor/UnrealEd/Private/AnimationEditorUtils.cpp
@@ -638,88 +638,115 @@ namespace AnimationEditorUtils
 		}
 	}
 
+	bool CanCreateAssetOfType(const UClass* InClass)
+	{
+		IAssetTools& AssetTools = FModuleManager::LoadModuleChecked<FAssetToolsModule>("AssetTools").Get();
+		return AssetTools.IsAssetClassSupported(InClass);
+	}
+
 	void FillCreateAssetMenu(FMenuBuilder& MenuBuilder, const TArray<TWeakObjectPtr<UObject>>& SkeletonsOrSkeletalMeshes, FAnimAssetCreated AssetCreated, bool bInContentBrowser)
 	{
 		MenuBuilder.BeginSection("CreateAnimAssets", LOCTEXT("CreateAnimAssetsMenuHeading", "Anim Assets"));
 		{
-			// only allow for content browser until we support multi assets so we can open new persona with this BP
-			MenuBuilder.AddMenuEntry(
-				LOCTEXT("Skeleton_NewAnimBlueprint", "Anim Blueprint"),
-				LOCTEXT("Skeleton_NewAnimBlueprintTooltip", "Creates an Anim Blueprint using the selected skeleton."),
-				FSlateIcon(FEditorStyle::GetStyleSetName(), "ClassIcon.AnimBlueprint"),
-				FUIAction(
-					FExecuteAction::CreateStatic(&CreateNewAnimBlueprint, SkeletonsOrSkeletalMeshes, AssetCreated, bInContentBrowser),
-					FCanExecuteAction()
-					)
-				);
-
-			MenuBuilder.AddMenuEntry(
-				LOCTEXT("Skeleton_NewAnimComposite", "Anim Composite"),
-				LOCTEXT("Skeleton_NewAnimCompositeTooltip", "Creates an AnimComposite using the selected skeleton."),
-				FSlateIcon(FEditorStyle::GetStyleSetName(), "ClassIcon.AnimComposite"),
-				FUIAction(
-					FExecuteAction::CreateStatic(&ExecuteNewAnimAsset<UAnimCompositeFactory, UAnimComposite>, SkeletonsOrSkeletalMeshes, FString("_Composite"), AssetCreated, bInContentBrowser),
-					FCanExecuteAction()
-					)
-				);
-
-			MenuBuilder.AddMenuEntry(
-				LOCTEXT("Skeleton_NewAnimMontage", "Anim Montage"),
-				LOCTEXT("Skeleton_NewAnimMontageTooltip", "Creates an AnimMontage using the selected skeleton."),
-				FSlateIcon(FEditorStyle::GetStyleSetName(), "ClassIcon.AnimMontage"),
-				FUIAction(
-					FExecuteAction::CreateStatic(&ExecuteNewAnimAsset<UAnimMontageFactory, UAnimMontage>, SkeletonsOrSkeletalMeshes, FString("_Montage"), AssetCreated, bInContentBrowser),
-					FCanExecuteAction()
-					)
-				);
+			if(CanCreateAssetOfType(UAnimBlueprint::StaticClass()))
+			{
+				// only allow for content browser until we support multi assets so we can open new persona with this BP
+				MenuBuilder.AddMenuEntry(
+					LOCTEXT("Skeleton_NewAnimBlueprint", "Anim Blueprint"),
+					LOCTEXT("Skeleton_NewAnimBlueprintTooltip", "Creates an Anim Blueprint using the selected skeleton."),
+					FSlateIcon(FEditorStyle::GetStyleSetName(), "ClassIcon.AnimBlueprint"),
+					FUIAction(
+						FExecuteAction::CreateStatic(&CreateNewAnimBlueprint, SkeletonsOrSkeletalMeshes, AssetCreated, bInContentBrowser),
+						FCanExecuteAction()
+						)
+					);
+			}
+
+			if(CanCreateAssetOfType(UAnimComposite::StaticClass()))
+			{
+				MenuBuilder.AddMenuEntry(
+					LOCTEXT("Skeleton_NewAnimComposite", "Anim Composite"),
+					LOCTEXT("Skeleton_NewAnimCompositeTooltip", "Creates an AnimComposite using the selected skeleton."),
+					FSlateIcon(FEditorStyle::GetStyleSetName(), "ClassIcon.AnimComposite"),
+					FUIAction(
+						FExecuteAction::CreateStatic(&ExecuteNewAnimAsset<UAnimCompositeFactory, UAnimComposite>, SkeletonsOrSkeletalMeshes, FString("_Composite"), AssetCreated, bInContentBrowser),
+						FCanExecuteAction()
+						)
+					);
+			}
+
+			if(CanCreateAssetOfType(UAnimMontage::StaticClass()))
+			{
+				MenuBuilder.AddMenuEntry(
+					LOCTEXT("Skeleton_NewAnimMontage", "Anim Montage"),
+					LOCTEXT("Skeleton_NewAnimMontageTooltip", "Creates an AnimMontage using the selected skeleton."),
+					FSlateIcon(FEditorStyle::GetStyleSetName(), "ClassIcon.AnimMontage"),
+					FUIAction(
+						FExecuteAction::CreateStatic(&ExecuteNewAnimAsset<UAnimMontageFactory, UAnimMontage>, SkeletonsOrSkeletalMeshes, FString("_Montage"), AssetCreated, bInContentBrowser),
+						FCanExecuteAction()
+						)
+					);
+			}
 		}
 		MenuBuilder.EndSection();
 
 		MenuBuilder.BeginSection("CreateBlendSpace", LOCTEXT("CreateBlendSpaceMenuHeading", "Blend Spaces"));
 		{
-			MenuBuilder.AddMenuEntry(
-				LOCTEXT("SkeletalMesh_New2DBlendspace", "Blend Space"),
-				LOCTEXT("SkeletalMesh_New2DBlendspaceTooltip", "Creates a Blend Space using the selected skeleton."),
-				FSlateIcon(FEditorStyle::GetStyleSetName(), "ClassIcon.BlendSpace"),
-				FUIAction(
-					FExecuteAction::CreateStatic(&ExecuteNewAnimAsset<UBlendSpaceFactoryNew, UBlendSpace>, SkeletonsOrSkeletalMeshes, FString("_BlendSpace"), AssetCreated, bInContentBrowser),
-					FCanExecuteAction()
-					)
-				);
-
-			MenuBuilder.AddMenuEntry(
-				LOCTEXT("SkeletalMesh_New1DBlendspace", "Blend Space 1D"),
-				LOCTEXT("SkeletalMesh_New1DBlendspaceTooltip", "Creates a 1D Blend Space using the selected skeleton."),
-				FSlateIcon(FEditorStyle::GetStyleSetName(), "ClassIcon.BlendSpace1D"),
-				FUIAction(
-					FExecuteAction::CreateStatic(&ExecuteNewAnimAsset<UBlendSpaceFactory1D, UBlendSpace1D>, SkeletonsOrSkeletalMeshes, FString("_BlendSpace1D"), AssetCreated, bInContentBrowser),
-					FCanExecuteAction()
-					)
-				);
+			if (CanCreateAssetOfType(UBlendSpace::StaticClass()))
+			{
+				MenuBuilder.AddMenuEntry(
+					LOCTEXT("SkeletalMesh_New2DBlendspace", "Blend Space"),
+					LOCTEXT("SkeletalMesh_New2DBlendspaceTooltip", "Creates a Blend Space using the selected skeleton."),
+					FSlateIcon(FEditorStyle::GetStyleSetName(), "ClassIcon.BlendSpace"),
+					FUIAction(
+						FExecuteAction::CreateStatic(&ExecuteNewAnimAsset<UBlendSpaceFactoryNew, UBlendSpace>, SkeletonsOrSkeletalMeshes, FString("_BlendSpace"), AssetCreated, bInContentBrowser),
+						FCanExecuteAction()
+						)
+					);
+			}
+
+			if (CanCreateAssetOfType(UBlendSpace1D::StaticClass()))
+			{
+				MenuBuilder.AddMenuEntry(
+					LOCTEXT("SkeletalMesh_New1DBlendspace", "Blend Space 1D"),
+					LOCTEXT("SkeletalMesh_New1DBlendspaceTooltip", "Creates a 1D Blend Space using the selected skeleton."),
+					FSlateIcon(FEditorStyle::GetStyleSetName(), "ClassIcon.BlendSpace1D"),
+					FUIAction(
+						FExecuteAction::CreateStatic(&ExecuteNewAnimAsset<UBlendSpaceFactory1D, UBlendSpace1D>, SkeletonsOrSkeletalMeshes, FString("_BlendSpace1D"), AssetCreated, bInContentBrowser),
+						FCanExecuteAction()
+						)
+					);
+			}
 		}
 		MenuBuilder.EndSection();
 
 		MenuBuilder.BeginSection("CreateAimOffset", LOCTEXT("CreateAimOffsetMenuHeading", "Aim Offsets"));
 		{
-			MenuBuilder.AddMenuEntry(
-				LOCTEXT("SkeletalMesh_New2DAimOffset", "Aim Offset"),
-				LOCTEXT("SkeletalMesh_New2DAimOffsetTooltip", "Creates a Aim Offset blendspace using the selected skeleton."),
-				FSlateIcon(),
-				FUIAction(
-					FExecuteAction::CreateStatic(&ExecuteNewAnimAsset<UAimOffsetBlendSpaceFactoryNew, UAimOffsetBlendSpace>, SkeletonsOrSkeletalMeshes, FString("_AimOffset2D"), AssetCreated, bInContentBrowser),
-					FCanExecuteAction()
-					)
-				);
-
-			MenuBuilder.AddMenuEntry(
-				LOCTEXT("SkeletalMesh_New1DAimOffset", "Aim Offset 1D"),
-				LOCTEXT("SkeletalMesh_New1DAimOffsetTooltip", "Creates a 1D Aim Offset blendspace using the selected skeleton."),
-				FSlateIcon(),
-				FUIAction(
-					FExecuteAction::CreateStatic(&ExecuteNewAnimAsset<UAimOffsetBlendSpaceFactory1D, UAimOffsetBlendSpace1D>, SkeletonsOrSkeletalMeshes, FString("_AimOffset1D"), AssetCreated, bInContentBrowser),
-					FCanExecuteAction()
-					)
-				);
+			if (CanCreateAssetOfType(UAimOffsetBlendSpace::StaticClass()))
+			{
+				MenuBuilder.AddMenuEntry(
+					LOCTEXT("SkeletalMesh_New2DAimOffset", "Aim Offset"),
+					LOCTEXT("SkeletalMesh_New2DAimOffsetTooltip", "Creates a Aim Offset blendspace using the selected skeleton."),
+					FSlateIcon(),
+					FUIAction(
+						FExecuteAction::CreateStatic(&ExecuteNewAnimAsset<UAimOffsetBlendSpaceFactoryNew, UAimOffsetBlendSpace>, SkeletonsOrSkeletalMeshes, FString("_AimOffset2D"), AssetCreated, bInContentBrowser),
+						FCanExecuteAction()
+						)
+					);
+			}
+
+			if (CanCreateAssetOfType(UAimOffsetBlendSpace1D::StaticClass()))
+			{
+				MenuBuilder.AddMenuEntry(
+					LOCTEXT("SkeletalMesh_New1DAimOffset", "Aim Offset 1D"),
+					LOCTEXT("SkeletalMesh_New1DAimOffsetTooltip", "Creates a 1D Aim Offset blendspace using the selected skeleton."),
+					FSlateIcon(),
+					FUIAction(
+						FExecuteAction::CreateStatic(&ExecuteNewAnimAsset<UAimOffsetBlendSpaceFactory1D, UAimOffsetBlendSpace1D>, SkeletonsOrSkeletalMeshes, FString("_AimOffset1D"), AssetCreated, bInContentBrowser),
+						FCanExecuteAction()
+						)
+					);
+			}
 		}
 		MenuBuilder.EndSection();
 	}
