Updated skeleton tree children ordering: sockets, preview meshes, and virtual bones are now placed at the head of a bone's children

diff --git a/Engine/Source/Editor/SkeletonEditor/Private/SkeletonTreeBuilder.cpp b/Engine/Source/Editor/SkeletonEditor/Private/SkeletonTreeBuilder.cpp
--- a/Engine/Source/Editor/SkeletonEditor/Private/SkeletonTreeBuilder.cpp
+++ b/Engine/Source/Editor/SkeletonEditor/Private/SkeletonTreeBuilder.cpp
@@ -355,7 +355,7 @@ void FSkeletonTreeBuilder::AddSocketsFromData(const TArray< USkeletalMeshSocket*
 			}
 		}
 
-		Output.Add(CreateSocketTreeItem(Socket, ParentType, bIsCustomized), Socket->BoneName, FSkeletonTreeBoneItem::GetTypeId());
+		Output.Add(CreateSocketTreeItem(Socket, ParentType, bIsCustomized), Socket->BoneName, FSkeletonTreeBoneItem::GetTypeId(), /*bAddToHead*/ true);
 	}
 }
 
@@ -365,7 +365,7 @@ void FSkeletonTreeBuilder::AddAttachedAssetContainer(const FPreviewAssetAttachCo
 	{
 		const FPreviewAttachedObjectPair& Pair = (*Iter);
 
-		Output.Add(CreateAttachedAssetTreeItem(Pair.GetAttachedObject(), Pair.AttachedTo), Pair.AttachedTo, { FSkeletonTreeBoneItem::GetTypeId(), FSkeletonTreeSocketItem::GetTypeId() });
+		Output.Add(CreateAttachedAssetTreeItem(Pair.GetAttachedObject(), Pair.AttachedTo), Pair.AttachedTo, { FSkeletonTreeBoneItem::GetTypeId(), FSkeletonTreeSocketItem::GetTypeId() }, /*bAddToHead*/ true);
 	}
 }
 
@@ -374,7 +374,7 @@ void FSkeletonTreeBuilder::AddVirtualBones(FSkeletonTreeBuilderOutput& Output)
 	const TArray<FVirtualBone>& VirtualBones = EditableSkeletonPtr.Pin()->GetSkeleton().GetVirtualBones();
 	for (const FVirtualBone& VirtualBone : VirtualBones)
 	{
-		Output.Add(CreateVirtualBoneTreeItem(VirtualBone.VirtualBoneName), VirtualBone.SourceBoneName, { FSkeletonTreeBoneItem::GetTypeId(), FSkeletonTreeVirtualBoneItem::GetTypeId() });
+		Output.Add(CreateVirtualBoneTreeItem(VirtualBone.VirtualBoneName), VirtualBone.SourceBoneName, { FSkeletonTreeBoneItem::GetTypeId(), FSkeletonTreeVirtualBoneItem::GetTypeId() }, /*bAddToHead*/ true);
 	}
 }
 
