Fixed bone transform clipboard pasting not working unless bone was previosly modified.

diff --git a/Engine/Source/Editor/SkeletonEditor/Private/BoneProxyDetailsCustomization.cpp b/Engine/Source/Editor/SkeletonEditor/Private/BoneProxyDetailsCustomization.cpp
--- a/Engine/Source/Editor/SkeletonEditor/Private/BoneProxyDetailsCustomization.cpp
+++ b/Engine/Source/Editor/SkeletonEditor/Private/BoneProxyDetailsCustomization.cpp
@@ -340,61 +340,60 @@ void FBoneProxyDetailsCustomization::CustomizeDetails(IDetailLayoutBuilder& Deta
 			{
 				if (UDebugSkelMeshComponent* Component = BoneProxy->SkelMeshComponent.Get())
 				{
-					if (FAnimNode_ModifyBone* ModifyBone = Component->PreviewInstance->FindModifiedBone(BoneProxy->BoneName))
+					FAnimNode_ModifyBone& ModifyBone = Component->PreviewInstance->ModifyBone(BoneProxy->BoneName);
+					
+					class FBoneProxyDetailsCustomization : public FOutputDevice
 					{
-						class FBoneProxyDetailsCustomization : public FOutputDevice
-						{
-						public:
+					public:
 
-							int32 NumErrors;
+						int32 NumErrors;
 
-							FBoneProxyDetailsCustomization()
-								: FOutputDevice()
-								, NumErrors(0)
-							{
-							}
+						FBoneProxyDetailsCustomization()
+							: FOutputDevice()
+							, NumErrors(0)
+						{
+						}
 
-							virtual void Serialize(const TCHAR* V, ELogVerbosity::Type Verbosity, const class FName& Category) override
-							{
-								NumErrors++;
-							}
-						};
+						virtual void Serialize(const TCHAR* V, ELogVerbosity::Type Verbosity, const class FName& Category) override
+						{
+							NumErrors++;
+						}
+					};
 
-						FBoneProxyDetailsCustomization ErrorPipe;
-										
-						switch(InComponent)
+					FBoneProxyDetailsCustomization ErrorPipe;
+									
+					switch(InComponent)
+					{
+						case ESlateTransformComponent::Location:
+						{
+							FVector Data = BoneProxy->Location;
+							TBaseStructure<FVector>::Get()->ImportText(*Content, &Data, nullptr, PPF_None, &ErrorPipe, TBaseStructure<FVector>::Get()->GetName(), true);
+							ModifyBone.Translation = Data;
+							break;
+						}
+						case ESlateTransformComponent::Rotation:
+						{
+							FRotator Data = BoneProxy->Rotation;
+							TBaseStructure<FRotator>::Get()->ImportText(*Content, &Data, nullptr, PPF_None, &ErrorPipe, TBaseStructure<FRotator>::Get()->GetName(), true);
+							ModifyBone.Rotation = Data;
+							break;
+						}
+						case ESlateTransformComponent::Scale:
+						{
+							FVector Data = BoneProxy->Scale;
+							TBaseStructure<FVector>::Get()->ImportText(*Content, &Data, nullptr, PPF_None, &ErrorPipe, TBaseStructure<FVector>::Get()->GetName(), true);
+							ModifyBone.Scale = Data;
+							break;
+						}
+						case ESlateTransformComponent::Max:
+						default:
 						{
-							case ESlateTransformComponent::Location:
-							{
-								FVector Data = BoneProxy->Location;
-								TBaseStructure<FVector>::Get()->ImportText(*Content, &Data, nullptr, PPF_None, &ErrorPipe, TBaseStructure<FVector>::Get()->GetName(), true);
-								ModifyBone->Translation = Data;
-								break;
-							}
-							case ESlateTransformComponent::Rotation:
-							{
-								FRotator Data = BoneProxy->Rotation;
-								TBaseStructure<FRotator>::Get()->ImportText(*Content, &Data, nullptr, PPF_None, &ErrorPipe, TBaseStructure<FRotator>::Get()->GetName(), true);
-								ModifyBone->Rotation = Data;
-								break;
-							}
-							case ESlateTransformComponent::Scale:
-							{
-								FVector Data = BoneProxy->Scale;
-								TBaseStructure<FVector>::Get()->ImportText(*Content, &Data, nullptr, PPF_None, &ErrorPipe, TBaseStructure<FVector>::Get()->GetName(), true);
-								ModifyBone->Scale = Data;
-								break;
-							}
-							case ESlateTransformComponent::Max:
-							default:
-							{
-								FEulerTransform Data = FEulerTransform::Identity;
-								TBaseStructure<FEulerTransform>::Get()->ImportText(*Content, &Data, nullptr, PPF_None, &ErrorPipe, TBaseStructure<FEulerTransform>::Get()->GetName(), true);
-								ModifyBone->Translation = Data.GetLocation();
-								ModifyBone->Rotation = Data.Rotator();
-								ModifyBone->Scale = Data.GetScale3D();
-								break;
-							}
+							FEulerTransform Data = FEulerTransform::Identity;
+							TBaseStructure<FEulerTransform>::Get()->ImportText(*Content, &Data, nullptr, PPF_None, &ErrorPipe, TBaseStructure<FEulerTransform>::Get()->GetName(), true);
+							ModifyBone.Translation = Data.GetLocation();
+							ModifyBone.Rotation = Data.Rotator();
+							ModifyBone.Scale = Data.GetScale3D();
+							break;
 						}
 					}
 				}
