Fixed repeated menu items (and crash) in skeleton tree blend profile picker

diff --git a/Engine/Source/Editor/SkeletonEditor/Private/SSkeletonTree.cpp b/Engine/Source/Editor/SkeletonEditor/Private/SSkeletonTree.cpp
--- a/Engine/Source/Editor/SkeletonEditor/Private/SSkeletonTree.cpp
+++ b/Engine/Source/Editor/SkeletonEditor/Private/SSkeletonTree.cpp
@@ -68,6 +68,7 @@
 #include "SPositiveActionButton.h"
 #include "ToolMenus.h"
 #include "ToolMenuMisc.h"
+#include "SkeletonTreeMenuContext.h"
 
 #define LOCTEXT_NAMESPACE "SSkeletonTree"
 
@@ -212,9 +213,9 @@ void SSkeletonTree::Construct(const FArguments& InArgs, const TSharedRef<FEditab
 	FSkeletonTreeCommands::Register();
 	BindCommands();
 
-	CreateBlendProfileMenu(nullptr);
-	CreateNewMenu();
-	CreateFilterMenu();
+	RegisterBlendProfileMenu();
+	RegisterNewMenu();
+	RegisterFilterMenu();
 
 	this->ChildSlot
 	[
@@ -1600,7 +1601,12 @@ void SSkeletonTree::HandlePackageReloaded(const EPackageReloadPhase InPackageRel
 
 TSharedRef<SWidget> SSkeletonTree::GetBlendProfileColumnMenuContent()
 {
-	return UToolMenus::Get()->GenerateWidget("SkeletonTree.BlendProfilesMenu", FToolMenuContext(UICommandList, Extenders));
+	FToolMenuContext MenuContext(UICommandList, Extenders);
+	USkeletonTreeMenuContext* SkeletonTreeMenuContext = NewObject<USkeletonTreeMenuContext>();
+	SkeletonTreeMenuContext->SkeletonTree = SharedThis(this);
+	MenuContext.AddObject(SkeletonTreeMenuContext);
+
+	return UToolMenus::Get()->GenerateWidget("SkeletonTree.BlendProfilesMenu", MenuContext);
 }
 
 void SSkeletonTree::ExpandTreeOnSelection(TSharedPtr<ISkeletonTreeItem> RowToExpand, bool bForce)
@@ -1616,62 +1622,79 @@ void SSkeletonTree::ExpandTreeOnSelection(TSharedPtr<ISkeletonTreeItem> RowToExp
 	}
 }
 
-void SSkeletonTree::CreateBlendProfileMenu(UToolMenu* InMenu)
+void SSkeletonTree::RegisterBlendProfileMenu()
 {
+	const FName MenuName("SkeletonTree.BlendProfilesMenu");
+	if (UToolMenus::Get()->IsMenuRegistered(MenuName))
+	{
+		return;
+	}
+
 	FToolMenuOwnerScoped OwnerScoped(this);
 
-	UToolMenu* Menu = InMenu ? InMenu : UToolMenus::Get()->ExtendMenu("SkeletonTree.BlendProfilesMenu");
+	UToolMenu* Menu = UToolMenus::Get()->RegisterMenu(MenuName);
 
 	Menu->AddDynamicSection(NAME_None,
-		FNewSectionConstructChoice(FNewToolMenuDelegate::CreateLambda([this](UToolMenu* InMenu)
+		FNewSectionConstructChoice(FNewToolMenuDelegate::CreateLambda([](UToolMenu* InMenu)
 		{
-			const FSkeletonTreeCommands& Actions = FSkeletonTreeCommands::Get();
-
-			FToolMenuSection& ProfilesSection = InMenu->AddSection(TEXT("BlendProfileActions"), LOCTEXT("BlendProfiles", "Blend Profiles"));
-
-			for (UBlendProfile* Profile : GetEditableSkeletonInternal()->GetBlendProfiles())
+			if(USkeletonTreeMenuContext* MenuContext = InMenu->Context.FindContext<USkeletonTreeMenuContext>())
 			{
-				ProfilesSection.AddMenuEntry(
-					Profile->GetFName(),
-					FText::FromName(Profile->GetFName()),
-					LOCTEXT("SelectBlendProfileTooltip", "Select this profile for editing."),
-					FSlateIcon(),
-					FToolUIActionChoice(
-						FUIAction(
-							FExecuteAction::CreateSP(BlendProfilePicker.ToSharedRef(), &SBlendProfilePicker::SetSelectedProfile, Profile, true),
-							FCanExecuteAction(),
-							FIsActionChecked::CreateSP(this, &SSkeletonTree::IsBlendProfileSelected, Profile->GetFName())
-						)
-					),
-					EUserInterfaceActionType::RadioButton
-				);
+				if(TSharedPtr<SSkeletonTree> SkeletonTree = MenuContext->SkeletonTree.Pin())
+				{
+					SkeletonTree->CreateBlendProfileMenu(InMenu);
+				}
 			}
+		})));
+}
 
-			if (BlendProfilePicker->GetSelectedBlendProfileName() != NAME_None)
-			{
-				FToolMenuSection& EditSection = InMenu->AddSection(TEXT("BlendProfileEdit"), LOCTEXT("EditBlendProfilesSection", "Edit"));
-
-				EditSection.AddMenuEntry(
-					"ClearBlendProfile",
-					LOCTEXT("Clear", "Clear Selected"),
-					LOCTEXT("Clear_ToolTip", "Clear the selected blend profile."),
-					FSlateIcon(),
-					FToolUIActionChoice(FUIAction(FExecuteAction::CreateSP(BlendProfilePicker.ToSharedRef(), &SBlendProfilePicker::OnClearSelection))));
-
-				EditSection.AddMenuEntry(
-					Actions.DeleteCurrentBlendProfile, 
-					FText::Format(LOCTEXT("DeleteBlendProfileLabel", "Delete {0}"),
-					FText::FromName(BlendProfilePicker->GetSelectedBlendProfileName())));
-			}
+void SSkeletonTree::CreateBlendProfileMenu(UToolMenu* InMenu)
+{
+	const FSkeletonTreeCommands& Actions = FSkeletonTreeCommands::Get();
 
-			{
-				FToolMenuSection& NewSection = InMenu->AddSection(TEXT("BlendProfileNew"), LOCTEXT("NewBlendProfiles", "New"));
-				NewSection.AddMenuEntry(Actions.CreateTimeBlendProfile);
-				NewSection.AddMenuEntry(Actions.CreateWeightBlendProfile);
-				NewSection.AddMenuEntry(Actions.CreateBlendMask);
-			}
-		})));
-} 
+	FToolMenuSection& ProfilesSection = InMenu->AddSection(TEXT("BlendProfileActions"), LOCTEXT("BlendProfiles", "Blend Profiles"));
+
+	for (UBlendProfile* Profile : GetEditableSkeletonInternal()->GetBlendProfiles())
+	{
+		ProfilesSection.AddMenuEntry(
+			Profile->GetFName(),
+			FText::FromName(Profile->GetFName()),
+			LOCTEXT("SelectBlendProfileTooltip", "Select this profile for editing."),
+			FSlateIcon(),
+			FToolUIActionChoice(
+				FUIAction(
+					FExecuteAction::CreateSP(BlendProfilePicker.ToSharedRef(), &SBlendProfilePicker::SetSelectedProfile, Profile, true),
+					FCanExecuteAction(),
+					FIsActionChecked::CreateSP(this, &SSkeletonTree::IsBlendProfileSelected, Profile->GetFName())
+				)
+			),
+			EUserInterfaceActionType::RadioButton
+		);
+	}
+
+	if (BlendProfilePicker->GetSelectedBlendProfileName() != NAME_None)
+	{
+		FToolMenuSection& EditSection = InMenu->AddSection(TEXT("BlendProfileEdit"), LOCTEXT("EditBlendProfilesSection", "Edit"));
+
+		EditSection.AddMenuEntry(
+			"ClearBlendProfile",
+			LOCTEXT("Clear", "Clear Selected"),
+			LOCTEXT("Clear_ToolTip", "Clear the selected blend profile."),
+			FSlateIcon(),
+			FToolUIActionChoice(FUIAction(FExecuteAction::CreateSP(BlendProfilePicker.ToSharedRef(), &SBlendProfilePicker::OnClearSelection))));
+
+		EditSection.AddMenuEntry(
+			Actions.DeleteCurrentBlendProfile,
+			FText::Format(LOCTEXT("DeleteBlendProfileLabel", "Delete {0}"),
+				FText::FromName(BlendProfilePicker->GetSelectedBlendProfileName())));
+	}
+
+	{
+		FToolMenuSection& NewSection = InMenu->AddSection(TEXT("BlendProfileNew"), LOCTEXT("NewBlendProfiles", "New"));
+		NewSection.AddMenuEntry(Actions.CreateTimeBlendProfile);
+		NewSection.AddMenuEntry(Actions.CreateWeightBlendProfile);
+		NewSection.AddMenuEntry(Actions.CreateBlendMask);
+	}
+}
 
 void SSkeletonTree::OnCreateBlendProfile(const EBlendProfileMode InMode)
 {
@@ -1778,13 +1801,19 @@ void SSkeletonTree::NotifyUser( FNotificationInfo& NotificationInfo )
 	}
 }
 
-void SSkeletonTree::CreateNewMenu()
+void SSkeletonTree::RegisterNewMenu()
 {
+	const FName MenuName("SkeletonTree.NewMenu");
+	if (UToolMenus::Get()->IsMenuRegistered(MenuName))
+	{
+		return;
+	}
+
 	FToolMenuOwnerScoped OwnerScoped(this);
 
 	const FSkeletonTreeCommands& Actions = FSkeletonTreeCommands::Get();
 
-	UToolMenu* Menu = UToolMenus::Get()->ExtendMenu("SkeletonTree.NewMenu");
+	UToolMenu* Menu = UToolMenus::Get()->RegisterMenu(MenuName);
 
 	{
 		FToolMenuSection& CreateSection = Menu->AddSection("CreateNew", LOCTEXT("SkeletonCreateNew", "Create"));
@@ -1810,13 +1839,19 @@ TSharedRef< SWidget > SSkeletonTree::CreateNewMenuWidget()
 	return UToolMenus::Get()->GenerateWidget("SkeletonTree.NewMenu", FToolMenuContext(UICommandList, Extenders)); 
 }
 
-void SSkeletonTree::CreateFilterMenu()
+void SSkeletonTree::RegisterFilterMenu()
 {
+	const FName MenuName("SkeletonTree.FilterMenu");
+	if (UToolMenus::Get()->IsMenuRegistered(MenuName))
+	{
+		return;
+	}
+
 	FToolMenuOwnerScoped OwnerScoped(this);
 
 	const FSkeletonTreeCommands& Actions = FSkeletonTreeCommands::Get();
 
-	UToolMenu* Menu = UToolMenus::Get()->ExtendMenu("SkeletonTree.FilterMenu");
+	UToolMenu* Menu = UToolMenus::Get()->RegisterMenu(MenuName);
 
 	{
 		FToolMenuSection& BlendProfilesSection = Menu->AddSection("BlendProfiles", LOCTEXT("BlendProfilesMenuHeading", "Blend Profiles"));
diff --git a/Engine/Source/Editor/SkeletonEditor/Private/SSkeletonTree.h b/Engine/Source/Editor/SkeletonEditor/Private/SSkeletonTree.h
--- a/Engine/Source/Editor/SkeletonEditor/Private/SSkeletonTree.h
+++ b/Engine/Source/Editor/SkeletonEditor/Private/SSkeletonTree.h
@@ -171,13 +171,13 @@ private:
 	TSharedRef< SWidget > CreateNewMenuWidget();
 
 	/** Called to create the add new menu */
-	void CreateNewMenu();
+	void RegisterNewMenu();
 
 	/** Called to display the filter menu */
 	TSharedRef< SWidget > CreateFilterMenuWidget();
 
 	/** Called to create the filter menu */
-	void CreateFilterMenu();
+	void RegisterFilterMenu();
 
 	/** Function to copy selected bone name to the clipboard */
 	void OnCopyBoneNames();
@@ -338,8 +338,11 @@ private:
 	/** Removes the active Blend Profile */
 	void OnDeleteCurrentBlendProfile();
 
-	/** Generate Blend Profile Menu */
-	void CreateBlendProfileMenu(UToolMenu* InMenu = nullptr);
+	/** Register Blend Profile Menu */
+	void RegisterBlendProfileMenu();
+
+	/** Create Blend Profile Menu */
+	void CreateBlendProfileMenu(UToolMenu* InMenu);
 
 	TSharedRef<SWidget> GetBlendProfileColumnMenuContent();
 
diff --git a/Engine/Source/Editor/SkeletonEditor/Private/SkeletonTreeMenuContext.h b/Engine/Source/Editor/SkeletonEditor/Private/SkeletonTreeMenuContext.h
new file mode 100644
--- /dev/null
+++ b/Engine/Source/Editor/SkeletonEditor/Private/SkeletonTreeMenuContext.h
@@ -0,0 +1,16 @@
+// Copyright Epic Games, Inc. All Rights Reserved.
+
+#pragma once
+
+#include "SSkeletonTree.h"
+#include "SkeletonTreeMenuContext.generated.h"
+
+UCLASS()
+class USkeletonTreeMenuContext : public UObject
+{
+	GENERATED_BODY()
+
+public:
+	/** The skeleton tree we reference */
+	TWeakPtr<SSkeletonTree> SkeletonTree;
+};
\ No newline at end of file
