SSkeletonTree : Fixed LOD bone reduction RequiredBones calculation, when done from the SkeletonTree

diff --git a/Engine/Source/Editor/SkeletonEditor/Private/SSkeletonTree.cpp b/Engine/Source/Editor/SkeletonEditor/Private/SSkeletonTree.cpp
--- a/Engine/Source/Editor/SkeletonEditor/Private/SSkeletonTree.cpp
+++ b/Engine/Source/Editor/SkeletonEditor/Private/SSkeletonTree.cpp
@@ -1130,13 +1130,17 @@ void SSkeletonTree::RemoveFromLOD(int32 LODIndex, bool bIncludeSelected, bool bI
 			{
 				for (int32 Index = LODIndex + 1; Index < TotalLOD; ++Index)
 				{
-					MeshUtilities.RemoveBonesFromMesh(PreviewMeshComponent->GetSkeletalMeshAsset(), Index, &BonesToRemove);
 					PreviewMeshComponent->GetSkeletalMeshAsset()->AddBoneToReductionSetting(Index, BonesToRemove);
+					// We don't pass BoneNamesToRemove, as AddBoneToReductionSetting has added them to the LODInfoArray[LODIndex].BonesToRemove
+					// Which will be used by RemoveBonesFromMesh if we pass null BonesNamesToRemove (else will just remove the newly deleted bones, which is wrong)
+					MeshUtilities.RemoveBonesFromMesh(PreviewMeshComponent->GetSkeletalMeshAsset(), Index, nullptr);
 				}
 			}
 
 			// remove from current LOD
-			MeshUtilities.RemoveBonesFromMesh(PreviewMeshComponent->GetSkeletalMeshAsset(), LODIndex, &BonesToRemove);
+			// We don't pass BoneNamesToRemove, as AddBoneToReductionSetting has added them to the LODInfoArray[LODIndex].BonesToRemove
+			// Which will be used by RemoveBonesFromMesh if we pass null BonesNamesToRemove (else will just remove the newly deleted bones, which is wrong)
+			MeshUtilities.RemoveBonesFromMesh(PreviewMeshComponent->GetSkeletalMeshAsset(), LODIndex, nullptr);
 		}
 		// update UI to reflect the change
 		OnLODSwitched();
