Fix for blend profile rename being reset when committing via clicking to lose focus
* No longer have to explicitly press Enter to commit to the name change.

diff --git a/Engine/Source/Editor/SkeletonEditor/Private/SBlendProfilePicker.cpp b/Engine/Source/Editor/SkeletonEditor/Private/SBlendProfilePicker.cpp
--- a/Engine/Source/Editor/SkeletonEditor/Private/SBlendProfilePicker.cpp
+++ b/Engine/Source/Editor/SkeletonEditor/Private/SBlendProfilePicker.cpp
@@ -481,7 +481,7 @@ void SBlendProfilePicker::OnCreateNewProfileComitted(const FText& NewName, EText
 {
 	FSlateApplication::Get().DismissAllMenus();
 
-	if(CommitType == ETextCommit::OnEnter && EditableSkeleton.IsValid())
+	if((CommitType == ETextCommit::OnEnter || CommitType == ETextCommit::OnUserMovedFocus) && EditableSkeleton.IsValid())
 	{
 		FScopedTransaction Transaction(LOCTEXT("Trans_NewProfile", "Create new blend profile."));
 
