follow the NumericValueChanged paradigm when pasting values + transaction

diff --git a/Engine/Source/Editor/SkeletonEditor/Private/BoneProxyDetailsCustomization.cpp b/Engine/Source/Editor/SkeletonEditor/Private/BoneProxyDetailsCustomization.cpp
--- a/Engine/Source/Editor/SkeletonEditor/Private/BoneProxyDetailsCustomization.cpp
+++ b/Engine/Source/Editor/SkeletonEditor/Private/BoneProxyDetailsCustomization.cpp
@@ -133,6 +133,38 @@ static TSharedRef<SWidget> BuildTransformFieldLabel(bool* bValuePtr, const FText
 	}
 }
 
+namespace
+{
+
+class FBoneProxyDetailsErrorPipe : public FOutputDevice
+{
+public:
+
+	int32 NumErrors;
+
+	FBoneProxyDetailsErrorPipe()
+		: FOutputDevice()
+		, NumErrors(0)
+	{
+	}
+
+	virtual void Serialize(const TCHAR* V, ELogVerbosity::Type Verbosity, const class FName& Category) override
+	{
+		NumErrors++;
+	}
+};
+
+template<typename DataType>
+bool GetDataFromContent(const FString& Content, DataType& OutData)
+{
+	FBoneProxyDetailsErrorPipe ErrorPipe;
+	static UScriptStruct* DataStruct = TBaseStructure<DataType>::Get();
+	DataStruct->ImportText(*Content, &OutData, nullptr, PPF_None, &ErrorPipe, DataStruct->GetName(), true);
+	return (ErrorPipe.NumErrors == 0);
+}
+	
+}
+
 void FBoneProxyDetailsCustomization::CustomizeDetails(IDetailLayoutBuilder& DetailBuilder)
 {
 	TArray<TWeakObjectPtr<UObject>> Objects;
@@ -336,63 +368,70 @@ void FBoneProxyDetailsCustomization::CustomizeDetails(IDetailLayoutBuilder& Deta
 				return;
 			}
 
+			FScopedTransaction Transaction(LOCTEXT("PasteTransform", "Paste Transform"));
+			
+			static const FName LocationPropertyName = GET_MEMBER_NAME_CHECKED(UBoneProxy, Location);
+			static const FName RotationPropertyName = GET_MEMBER_NAME_CHECKED(UBoneProxy, Rotation);
+			static const FName ScalePropertyName = GET_MEMBER_NAME_CHECKED(UBoneProxy, Scale);
+			static constexpr bool bIsCommit = true;
+
 			for (UBoneProxy* BoneProxy : BoneProxiesView)
 			{
-				if (UDebugSkelMeshComponent* Component = BoneProxy->SkelMeshComponent.Get())
+				if (const UDebugSkelMeshComponent* Component = BoneProxy->SkelMeshComponent.Get())
 				{
-					FAnimNode_ModifyBone& ModifyBone = Component->PreviewInstance->ModifyBone(BoneProxy->BoneName);
-					
-					class FBoneProxyDetailsCustomization : public FOutputDevice
-					{
-					public:
-
-						int32 NumErrors;
-
-						FBoneProxyDetailsCustomization()
-							: FOutputDevice()
-							, NumErrors(0)
-						{
-						}
-
-						virtual void Serialize(const TCHAR* V, ELogVerbosity::Type Verbosity, const class FName& Category) override
-						{
-							NumErrors++;
-						}
-					};
-
-					FBoneProxyDetailsCustomization ErrorPipe;
-									
 					switch(InComponent)
 					{
 						case ESlateTransformComponent::Location:
 						{
 							FVector Data = BoneProxy->Location;
-							TBaseStructure<FVector>::Get()->ImportText(*Content, &Data, nullptr, PPF_None, &ErrorPipe, TBaseStructure<FVector>::Get()->GetName(), true);
-							ModifyBone.Translation = Data;
+							if (GetDataFromContent(Content, Data))
+							{
+								BoneProxy->OnPreEditChange(LocationPropertyName, bIsCommit);
+								BoneProxy->Location = Data;
+								BoneProxy->OnPostEditChangeProperty(LocationPropertyName, bIsCommit);
+							}
 							break;
 						}
 						case ESlateTransformComponent::Rotation:
 						{
 							FRotator Data = BoneProxy->Rotation;
-							TBaseStructure<FRotator>::Get()->ImportText(*Content, &Data, nullptr, PPF_None, &ErrorPipe, TBaseStructure<FRotator>::Get()->GetName(), true);
-							ModifyBone.Rotation = Data;
+							if (GetDataFromContent(Content, Data))
+							{
+								BoneProxy->OnPreEditChange(RotationPropertyName, bIsCommit);
+								BoneProxy->Rotation = Data;
+								BoneProxy->OnPostEditChangeProperty(RotationPropertyName, bIsCommit);
+							}
 							break;
 						}
 						case ESlateTransformComponent::Scale:
 						{
 							FVector Data = BoneProxy->Scale;
-							TBaseStructure<FVector>::Get()->ImportText(*Content, &Data, nullptr, PPF_None, &ErrorPipe, TBaseStructure<FVector>::Get()->GetName(), true);
-							ModifyBone.Scale = Data;
+							if (GetDataFromContent(Content, Data))
+							{
+								BoneProxy->OnPreEditChange(ScalePropertyName, bIsCommit);
+								BoneProxy->Scale = Data;
+								BoneProxy->OnPostEditChangeProperty(ScalePropertyName, bIsCommit);
+							}
 							break;
 						}
 						case ESlateTransformComponent::Max:
 						default:
 						{
 							FEulerTransform Data = FEulerTransform::Identity;
-							TBaseStructure<FEulerTransform>::Get()->ImportText(*Content, &Data, nullptr, PPF_None, &ErrorPipe, TBaseStructure<FEulerTransform>::Get()->GetName(), true);
-							ModifyBone.Translation = Data.GetLocation();
-							ModifyBone.Rotation = Data.Rotator();
-							ModifyBone.Scale = Data.GetScale3D();
+							if (GetDataFromContent(Content, Data))
+							{
+								BoneProxy->OnPreEditChange(LocationPropertyName, bIsCommit);
+								BoneProxy->Location = Data.GetLocation();
+								BoneProxy->OnPostEditChangeProperty(LocationPropertyName, bIsCommit);
+								
+								BoneProxy->OnPreEditChange(RotationPropertyName, bIsCommit);
+								BoneProxy->Rotation = Data.Rotator();
+								BoneProxy->OnPostEditChangeProperty(RotationPropertyName, bIsCommit);
+								
+								BoneProxy->OnPreEditChange(ScalePropertyName, bIsCommit);
+								BoneProxy->Scale = Data.GetScale3D();
+								BoneProxy->OnPostEditChangeProperty(ScalePropertyName, bIsCommit);
+							}
 							break;
 						}
 					}
