Different blend profile modes have their own section in pickers

diff --git a/Engine/Source/Editor/SkeletonEditor/Private/SBlendProfilePicker.cpp b/Engine/Source/Editor/SkeletonEditor/Private/SBlendProfilePicker.cpp
--- a/Engine/Source/Editor/SkeletonEditor/Private/SBlendProfilePicker.cpp
+++ b/Engine/Source/Editor/SkeletonEditor/Private/SBlendProfilePicker.cpp
@@ -393,38 +393,55 @@ TSharedRef<SWidget> SBlendProfilePicker::GetMenuContent()
 		MenuBuilder.EndSection();
 	}
 
-	const static FText AvailableModesHeadings[] = { LOCTEXT("BlendProfiles", "Available Blend Profiles"), LOCTEXT("BlendMasks", "Available Blend Masks") };
-	const static EBlendProfilePickerMode AvailableModesFilters[] = { EBlendProfilePickerMode::BlendProfile, EBlendProfilePickerMode::BlendMask };
-	static_assert(UE_ARRAY_COUNT(AvailableModesHeadings) == UE_ARRAY_COUNT(AvailableModesFilters), "Arrays must be of same size");
+	const static FText BlendProfileHeaders[] =
+	{
+		LOCTEXT("BlendProfilesTimeBased", "Blend Profiles - Time"),
+		LOCTEXT("BlendProfilesWeightBased", "Blend Profiles - Weight"),
+		LOCTEXT("BlendMasks", "Blend Masks")
+	};
+
+	UEnum* ModeEnum = StaticEnum<EBlendProfileMode>();
+	check(ModeEnum);
+	// Last enum entry is _MAX
+	int32 NumEnums = ModeEnum->NumEnums() - 1;
+
+	TArray<TArray<UBlendProfile*>> BlendProfilesFiltered;
+	BlendProfilesFiltered.SetNum(NumEnums);
+
+	// Build a filtered profile list by mode
+	for (UBlendProfile* Profile : EditableSkeleton->GetBlendProfiles())
+	{
+		if (Profile)
+		{
+			BlendProfilesFiltered[ModeEnum->GetIndexByValue((int64)Profile->GetMode())].Add(Profile);
+		}
+	}
 
-	for (int32 ModeIndex = 0; ModeIndex < UE_ARRAY_COUNT(AvailableModesHeadings); ++ModeIndex)
+	const bool bSupportsBlendMasks = EnumHasAnyFlags(SupportedBlendProfileModes, EBlendProfilePickerMode::BlendMask);
+	const bool bSupportsBlendProfiles = EnumHasAnyFlags(SupportedBlendProfileModes, EBlendProfilePickerMode::BlendProfile);
+	for (int32 ModeIndex = 0; ModeIndex < NumEnums; ++ModeIndex)
 	{
-		// Skip section if we don't support it
-		if (EnumHasAnyFlags(SupportedBlendProfileModes, AvailableModesFilters[ModeIndex]))
+		EBlendProfileMode IndexMode = (EBlendProfileMode)ModeEnum->GetValueByIndex(ModeIndex);
+		if ((IndexMode == EBlendProfileMode::BlendMask && bSupportsBlendMasks)
+			|| (IndexMode != EBlendProfileMode::BlendMask && bSupportsBlendProfiles))
 		{
 			// Note: Section won't get populated if there are no available items for this mode type
-			MenuBuilder.BeginSection(NAME_None, AvailableModesHeadings[ModeIndex]);
+			MenuBuilder.BeginSection(NAME_None, BlendProfileHeaders[ModeIndex]);
+			for (UBlendProfile* Profile : BlendProfilesFiltered[ModeIndex])
 			{
-				if (EditableSkeleton.IsValid())
+				if (IndexMode == Profile->GetMode())
 				{
-					for (UBlendProfile* Profile : EditableSkeleton->GetBlendProfiles())
-					{
-						// Only display blend profiles w/ a matching mode
-						if ((AvailableModesFilters[ModeIndex] == EBlendProfilePickerMode::BlendMask) ? Profile->IsBlendMask() : !Profile->IsBlendMask())
-						{
-							MenuBuilder.AddWidget(
-								SNew(SBlendProfileMenuEntry)
-								.LabelOverride(FText::FromString(Profile->GetName()))
-								.OnOpenClickedDelegate(FExecuteAction::CreateSP(this, &SBlendProfilePicker::OnProfileSelected, Profile->GetFName()))
-								.OnRemoveClickedDelegate(FExecuteAction::CreateSP(this, &SBlendProfilePicker::OnProfileRemoved, Profile->GetFName()))
-								.OnProfileModeChangedDelegate(SBlendProfileMenuEntry::FBlendProfileModeChanged::CreateSP(this, &SBlendProfilePicker::OnProfileModeChanged, Profile->GetFName()))
-								.BlendProfile(MakeWeakObjectPtr(Profile))
-								.AllowModify(bAllowModify),
-								FText(),
-								true
-							);
-						}
-					}
+					MenuBuilder.AddWidget(
+						SNew(SBlendProfileMenuEntry)
+						.LabelOverride(FText::FromString(Profile->GetName()))
+						.OnOpenClickedDelegate(FExecuteAction::CreateSP(this, &SBlendProfilePicker::OnProfileSelected, Profile->GetFName()))
+						.OnRemoveClickedDelegate(FExecuteAction::CreateSP(this, &SBlendProfilePicker::OnProfileRemoved, Profile->GetFName()))
+						.OnProfileModeChangedDelegate(SBlendProfileMenuEntry::FBlendProfileModeChanged::CreateSP(this, &SBlendProfilePicker::OnProfileModeChanged, Profile->GetFName()))
+						.BlendProfile(MakeWeakObjectPtr(Profile))
+						.AllowModify(bAllowModify),
+						FText(),
+						true
+					);
 				}
 			}
 		}
diff --git a/Engine/Source/Editor/SkeletonEditor/Private/SSkeletonTree.cpp b/Engine/Source/Editor/SkeletonEditor/Private/SSkeletonTree.cpp
--- a/Engine/Source/Editor/SkeletonEditor/Private/SSkeletonTree.cpp
+++ b/Engine/Source/Editor/SkeletonEditor/Private/SSkeletonTree.cpp
@@ -1654,25 +1654,56 @@ void SSkeletonTree::RegisterBlendProfileMenu()
 void SSkeletonTree::CreateBlendProfileMenu(UToolMenu* InMenu)
 {
 	const FSkeletonTreeCommands& Actions = FSkeletonTreeCommands::Get();
+	
+	static const FName BlendProfileSectionNames[]
+	{
+		TEXT("BlendProfileTimeActions"),
+		TEXT("BlendProfileWeightActions"),
+		TEXT("BlendMaskActions")
+	};
 
-	FToolMenuSection& ProfilesSection = InMenu->AddSection(TEXT("BlendProfileActions"), LOCTEXT("BlendProfiles", "Blend Profiles"));
+	InMenu->AddSection(BlendProfileSectionNames[0], LOCTEXT("BlendProfilesTime", "Blend Profiles - Time"));
+	InMenu->AddSection(BlendProfileSectionNames[1], LOCTEXT("BlendProfilesWeight", "Blend Profiles - Weight"));
+	InMenu->AddSection(BlendProfileSectionNames[2], LOCTEXT("BlendProfiles", "Blend Masks"));
+	FToolMenuSection* BlendProfileSections[] = 
+	{
+		InMenu->FindSection(BlendProfileSectionNames[0]),
+		InMenu->FindSection(BlendProfileSectionNames[1]),
+		InMenu->FindSection(BlendProfileSectionNames[2])
+	};
+
+	static const FText SelectBlendProfileToolTipText = LOCTEXT("SelectBlendProfileTooltip", "Select this blend profile for editing.");
+	static const FText SelectBlendMaskToolTipText = LOCTEXT("SelectBlendProfileTooltip", "Select this blend mask for editing.");
+	static const FText SelectBlendProfileToolTipTexts[] = 
+	{
+		SelectBlendProfileToolTipText,
+		SelectBlendProfileToolTipText,
+		SelectBlendMaskToolTipText
+	};	
+
+	UEnum* ModeEnum = StaticEnum<EBlendProfileMode>();
+	check(ModeEnum);
 
 	for (UBlendProfile* Profile : GetEditableSkeletonInternal()->GetBlendProfiles())
 	{
-		ProfilesSection.AddMenuEntry(
-			Profile->GetFName(),
-			FText::FromName(Profile->GetFName()),
-			LOCTEXT("SelectBlendProfileTooltip", "Select this profile for editing."),
-			FSlateIcon(),
-			FToolUIActionChoice(
-				FUIAction(
-					FExecuteAction::CreateSP(BlendProfilePicker.ToSharedRef(), &SBlendProfilePicker::SetSelectedProfile, Profile, true),
-					FCanExecuteAction(),
-					FIsActionChecked::CreateSP(this, &SSkeletonTree::IsBlendProfileSelected, Profile->GetFName())
-				)
-			),
-			EUserInterfaceActionType::RadioButton
-		);
+		if (Profile)
+		{
+			int32 EnumIndex = ModeEnum->GetIndexByValue((int64)Profile->GetMode());
+			BlendProfileSections[EnumIndex]->AddMenuEntry(
+						Profile->GetFName(),
+						FText::FromName(Profile->GetFName()),
+						SelectBlendProfileToolTipTexts[EnumIndex],
+						FSlateIcon(),
+						FToolUIActionChoice(
+							FUIAction(
+								FExecuteAction::CreateSP(BlendProfilePicker.ToSharedRef(), &SBlendProfilePicker::SetSelectedProfile, Profile, true),
+								FCanExecuteAction(),
+								FIsActionChecked::CreateSP(this, &SSkeletonTree::IsBlendProfileSelected, Profile->GetFName())
+							)
+						),
+						EUserInterfaceActionType::RadioButton
+					);
+		}
 	}
 
 	if (BlendProfilePicker->GetSelectedBlendProfileName() != NAME_None)
@@ -1682,7 +1713,7 @@ void SSkeletonTree::CreateBlendProfileMenu(UToolMenu* InMenu)
 		EditSection.AddMenuEntry(
 			"ClearBlendProfile",
 			LOCTEXT("Clear", "Clear Selected"),
-			LOCTEXT("Clear_ToolTip", "Clear the selected blend profile."),
+			LOCTEXT("Clear_ToolTip", "Clear the selected blend profile/mask."),
 			FSlateIcon(),
 			FToolUIActionChoice(FUIAction(FExecuteAction::CreateSP(BlendProfilePicker.ToSharedRef(), &SBlendProfilePicker::OnClearSelection))));
 
