Fix CIS localization warning in SSkeletonTree

diff --git a/Engine/Source/Editor/SkeletonEditor/Private/SSkeletonTree.cpp b/Engine/Source/Editor/SkeletonEditor/Private/SSkeletonTree.cpp
--- a/Engine/Source/Editor/SkeletonEditor/Private/SSkeletonTree.cpp
+++ b/Engine/Source/Editor/SkeletonEditor/Private/SSkeletonTree.cpp
@@ -1673,7 +1673,7 @@ void SSkeletonTree::CreateBlendProfileMenu(UToolMenu* InMenu)
 	};
 
 	static const FText SelectBlendProfileToolTipText = LOCTEXT("SelectBlendProfileTooltip", "Select this blend profile for editing.");
-	static const FText SelectBlendMaskToolTipText = LOCTEXT("SelectBlendProfileTooltip", "Select this blend mask for editing.");
+	static const FText SelectBlendMaskToolTipText = LOCTEXT("SelectBlendMaskTooltip", "Select this blend mask for editing.");
 	static const FText SelectBlendProfileToolTipTexts[] = 
 	{
 		SelectBlendProfileToolTipText,
