Fix blend profiles submenu not appearing on a subsequent opening of a skeleton

Correctly apply context to tool menu

diff --git a/Engine/Source/Editor/SkeletonEditor/Private/SSkeletonTree.cpp b/Engine/Source/Editor/SkeletonEditor/Private/SSkeletonTree.cpp
--- a/Engine/Source/Editor/SkeletonEditor/Private/SSkeletonTree.cpp
+++ b/Engine/Source/Editor/SkeletonEditor/Private/SSkeletonTree.cpp
@@ -1641,18 +1641,24 @@ void SSkeletonTree::RegisterBlendProfileMenu()
 	Menu->AddDynamicSection(NAME_None,
 		FNewSectionConstructChoice(FNewToolMenuDelegate::CreateLambda([](UToolMenu* InMenu)
 		{
-			if(USkeletonTreeMenuContext* MenuContext = InMenu->Context.FindContext<USkeletonTreeMenuContext>())
-			{
-				if(TSharedPtr<SSkeletonTree> SkeletonTree = MenuContext->SkeletonTree.Pin())
-				{
-					SkeletonTree->CreateBlendProfileMenu(InMenu);
-				}
-			}
+			CreateBlendProfileMenu(InMenu);
 		})));
 }
 
 void SSkeletonTree::CreateBlendProfileMenu(UToolMenu* InMenu)
 {
+	USkeletonTreeMenuContext* MenuContext = InMenu->Context.FindContext<USkeletonTreeMenuContext>();
+	if(MenuContext == nullptr)
+	{
+		return;
+	}
+
+	TSharedPtr<SSkeletonTree> SkeletonTree = MenuContext->SkeletonTree.Pin();
+	if(!SkeletonTree.IsValid())
+	{
+		return;
+	}
+	
 	const FSkeletonTreeCommands& Actions = FSkeletonTreeCommands::Get();
 	
 	static const FName BlendProfileSectionNames[]
@@ -1684,7 +1690,7 @@ void SSkeletonTree::CreateBlendProfileMenu(UToolMenu* InMenu)
 	UEnum* ModeEnum = StaticEnum<EBlendProfileMode>();
 	check(ModeEnum);
 
-	for (UBlendProfile* Profile : GetEditableSkeletonInternal()->GetBlendProfiles())
+	for (UBlendProfile* Profile : SkeletonTree->GetEditableSkeletonInternal()->GetBlendProfiles())
 	{
 		if (Profile)
 		{
@@ -1696,9 +1702,9 @@ void SSkeletonTree::CreateBlendProfileMenu(UToolMenu* InMenu)
 						FSlateIcon(),
 						FToolUIActionChoice(
 							FUIAction(
-								FExecuteAction::CreateSP(BlendProfilePicker.ToSharedRef(), &SBlendProfilePicker::SetSelectedProfile, Profile, true),
+								FExecuteAction::CreateSP(SkeletonTree->BlendProfilePicker.ToSharedRef(), &SBlendProfilePicker::SetSelectedProfile, Profile, true),
 								FCanExecuteAction(),
-								FIsActionChecked::CreateSP(this, &SSkeletonTree::IsBlendProfileSelected, Profile->GetFName())
+								FIsActionChecked::CreateSP(SkeletonTree.Get(), &SSkeletonTree::IsBlendProfileSelected, Profile->GetFName())
 							)
 						),
 						EUserInterfaceActionType::RadioButton
@@ -1706,7 +1712,7 @@ void SSkeletonTree::CreateBlendProfileMenu(UToolMenu* InMenu)
 		}
 	}
 
-	if (BlendProfilePicker->GetSelectedBlendProfileName() != NAME_None)
+	if (SkeletonTree->BlendProfilePicker->GetSelectedBlendProfileName() != NAME_None)
 	{
 		FToolMenuSection& EditSection = InMenu->AddSection(TEXT("BlendProfileEdit"), LOCTEXT("EditBlendProfilesSection", "Edit"));
 
@@ -1715,12 +1721,12 @@ void SSkeletonTree::CreateBlendProfileMenu(UToolMenu* InMenu)
 			LOCTEXT("Clear", "Clear Selected"),
 			LOCTEXT("Clear_ToolTip", "Clear the selected blend profile/mask."),
 			FSlateIcon(),
-			FToolUIActionChoice(FUIAction(FExecuteAction::CreateSP(BlendProfilePicker.ToSharedRef(), &SBlendProfilePicker::OnClearSelection))));
+			FToolUIActionChoice(FUIAction(FExecuteAction::CreateSP(SkeletonTree->BlendProfilePicker.ToSharedRef(), &SBlendProfilePicker::OnClearSelection))));
 
 		EditSection.AddMenuEntry(
 			Actions.DeleteCurrentBlendProfile,
 			FText::Format(LOCTEXT("DeleteBlendProfileLabel", "Delete {0}"),
-				FText::FromName(BlendProfilePicker->GetSelectedBlendProfileName())));
+				FText::FromName(SkeletonTree->BlendProfilePicker->GetSelectedBlendProfileName())));
 	}
 
 	{
@@ -1871,7 +1877,12 @@ void SSkeletonTree::RegisterNewMenu()
 
 TSharedRef< SWidget > SSkeletonTree::CreateNewMenuWidget()
 {
-	return UToolMenus::Get()->GenerateWidget("SkeletonTree.NewMenu", FToolMenuContext(UICommandList, Extenders)); 
+	FToolMenuContext MenuContext(UICommandList, Extenders);
+	USkeletonTreeMenuContext* SkeletonTreeMenuContext = NewObject<USkeletonTreeMenuContext>();
+	SkeletonTreeMenuContext->SkeletonTree = SharedThis(this);
+	MenuContext.AddObject(SkeletonTreeMenuContext);
+	
+	return UToolMenus::Get()->GenerateWidget("SkeletonTree.NewMenu", MenuContext); 
 }
 
 void SSkeletonTree::RegisterFilterMenu()
@@ -1894,7 +1905,7 @@ void SSkeletonTree::RegisterFilterMenu()
 			"BlendProfiles", 
 			LOCTEXT("BlendProfilesSubMenu", "Blend Profiles"),
 			LOCTEXT("BlendProfilesSubMenuTooltip", "Edit Blend Profiles in this Skeleton"), 
-			FNewToolMenuChoice(FNewToolMenuDelegate::CreateSP(this, &SSkeletonTree::CreateBlendProfileMenu)));
+			FNewToolMenuChoice(FNewToolMenuDelegate::CreateStatic(&SSkeletonTree::CreateBlendProfileMenu)));
 	}
 
 	{
@@ -1926,7 +1937,12 @@ void SSkeletonTree::RegisterFilterMenu()
 
 TSharedRef< SWidget > SSkeletonTree::CreateFilterMenuWidget()
 {
-	return UToolMenus::Get()->GenerateWidget("SkeletonTree.FilterMenu", FToolMenuContext(UICommandList, Extenders));
+	FToolMenuContext MenuContext(UICommandList, Extenders);
+	USkeletonTreeMenuContext* SkeletonTreeMenuContext = NewObject<USkeletonTreeMenuContext>();
+	SkeletonTreeMenuContext->SkeletonTree = SharedThis(this);
+	MenuContext.AddObject(SkeletonTreeMenuContext);
+
+	return UToolMenus::Get()->GenerateWidget("SkeletonTree.FilterMenu", MenuContext);
 }
 
 void SSkeletonTree::SetBoneFilter( EBoneFilter InBoneFilter )
diff --git a/Engine/Source/Editor/SkeletonEditor/Private/SSkeletonTree.h b/Engine/Source/Editor/SkeletonEditor/Private/SSkeletonTree.h
--- a/Engine/Source/Editor/SkeletonEditor/Private/SSkeletonTree.h
+++ b/Engine/Source/Editor/SkeletonEditor/Private/SSkeletonTree.h
@@ -342,7 +342,7 @@ private:
 	void RegisterBlendProfileMenu();
 
 	/** Create Blend Profile Menu */
-	void CreateBlendProfileMenu(UToolMenu* InMenu);
+	static void CreateBlendProfileMenu(UToolMenu* InMenu);
 
 	TSharedRef<SWidget> GetBlendProfileColumnMenuContent();
 
