Fix for blend profile slider transaction spam each time the slider moves slightly.
* A transaction is now made from when the slider begins moving to when the value is committed.

diff --git a/Engine/Source/Editor/SkeletonEditor/Private/SkeletonTreeBoneItem.cpp b/Engine/Source/Editor/SkeletonEditor/Private/SkeletonTreeBoneItem.cpp
--- a/Engine/Source/Editor/SkeletonEditor/Private/SkeletonTreeBoneItem.cpp
+++ b/Engine/Source/Editor/SkeletonEditor/Private/SkeletonTreeBoneItem.cpp
@@ -19,6 +19,8 @@
 #include "Framework/MultiBox/MultiBoxBuilder.h"
 #include "Rendering/SkeletalMeshRenderData.h"
 #include "UObject/Package.h"
+#include "Editor.h"
+#include "ScopedTransaction.h"
 
 #define LOCTEXT_NAMESPACE "FSkeletonTreeBoneItem"
 
@@ -27,6 +29,7 @@ FSkeletonTreeBoneItem::FSkeletonTreeBoneItem(const FName& InBoneName, const TSha
 	, BoneName(InBoneName)
 	, bWeightedBone(false)
 	, bRequiredBone(false)
+	, bBlendSliderStartedTransaction(false)
 {
 	static const FString BoneProxyPrefix(TEXT("BONEPROXY_"));
 
@@ -138,7 +141,10 @@ TSharedRef< SWidget > FSkeletonTreeBoneItem::GenerateWidgetForDataColumn(const F
 				.MaxSliderValue(this, &FSkeletonTreeBoneItem::GetBlendProfileMaxSliderValue)
 				.Value(this, &FSkeletonTreeBoneItem::GetBoneBlendProfileScale)
 				.OnValueCommitted(this, &FSkeletonTreeBoneItem::OnBlendSliderCommitted)
-				.OnValueChanged(this, &FSkeletonTreeBoneItem::OnBlendSliderCommitted, ETextCommit::OnEnter)
+				.OnValueChanged(this, &FSkeletonTreeBoneItem::OnBlendSliderChanged)
+				.OnBeginSliderMovement(this, &FSkeletonTreeBoneItem::OnBeginBlendSliderMovement)
+				.OnEndSliderMovement(this, &FSkeletonTreeBoneItem::OnEndBlendSliderMovement)
+				.ClearKeyboardFocusOnCommit(true)
 			];
 	}
 
@@ -245,12 +251,6 @@ void FSkeletonTreeBoneItem::SetBoneTranslationRetargetingMode(EBoneTranslationRe
 	GetEditableSkeleton()->SetBoneTranslationRetargetingMode(BoneName, NewRetargetingMode);
 }
 
-void FSkeletonTreeBoneItem::SetBoneBlendProfileScale(float NewScale, bool bRecurse)
-{
-	FName BlendProfileName = GetSkeletonTree()->GetSelectedBlendProfile()->GetFName();
-	GetEditableSkeleton()->SetBlendProfileScale(BlendProfileName, BoneName, NewScale, bRecurse);
-}
-
 FSlateFontInfo FSkeletonTreeBoneItem::GetBoneTextFont() const
 {
 	if (!bRequiredBone)
@@ -376,9 +376,56 @@ FText FSkeletonTreeBoneItem::GetBoneToolTip()
 	return ToolTip;
 }
 
+void FSkeletonTreeBoneItem::OnBeginBlendSliderMovement()
+{
+	if (bBlendSliderStartedTransaction == false)
+	{
+		bBlendSliderStartedTransaction = true;
+		GEditor->BeginTransaction(LOCTEXT("BlendSliderTransation", "Modify Blend Profile Value"));
+
+		const FName& BlendProfileName = GetSkeletonTree()->GetSelectedBlendProfile()->GetFName();
+		UBlendProfile* BlendProfile = GetEditableSkeleton()->GetBlendProfile(BlendProfileName);
+
+		if (BlendProfile)
+		{
+			BlendProfile->SetFlags(RF_Transactional);
+			BlendProfile->Modify();
+		}
+	}
+}
+void FSkeletonTreeBoneItem::OnEndBlendSliderMovement(float NewValue)
+{
+	if (bBlendSliderStartedTransaction)
+	{
+		GEditor->EndTransaction();
+		bBlendSliderStartedTransaction = false;
+	}
+}
+
 void FSkeletonTreeBoneItem::OnBlendSliderCommitted(float NewValue, ETextCommit::Type CommitType)
 {
-	SetBoneBlendProfileScale(NewValue, false);
+	FName BlendProfileName = GetSkeletonTree()->GetSelectedBlendProfile()->GetFName();
+	UBlendProfile* BlendProfile = GetEditableSkeleton()->GetBlendProfile(BlendProfileName);
+
+	if (BlendProfile)
+	{
+		FScopedTransaction(LOCTEXT("SetBlendProfileValue", "Set Blend Profile Value"));
+		BlendProfile->SetFlags(RF_Transactional);
+		BlendProfile->Modify();
+
+		BlendProfile->SetBoneBlendScale(BoneName, NewValue, false, true);
+	}
+}
+
+void FSkeletonTreeBoneItem::OnBlendSliderChanged(float NewValue)
+{
+	const FName& BlendProfileName = GetSkeletonTree()->GetSelectedBlendProfile()->GetFName();
+	UBlendProfile* BlendProfile = GetEditableSkeleton()->GetBlendProfile(BlendProfileName);
+	
+	if (BlendProfile)
+	{
+		BlendProfile->SetBoneBlendScale(BoneName, NewValue, false, true);
+	}
 }
 
 void FSkeletonTreeBoneItem::HandleDragEnter(const FDragDropEvent& DragDropEvent)
diff --git a/Engine/Source/Editor/SkeletonEditor/Private/SkeletonTreeBoneItem.h b/Engine/Source/Editor/SkeletonEditor/Private/SkeletonTreeBoneItem.h
--- a/Engine/Source/Editor/SkeletonEditor/Private/SkeletonTreeBoneItem.h
+++ b/Engine/Source/Editor/SkeletonEditor/Private/SkeletonTreeBoneItem.h
@@ -77,15 +77,21 @@ private:
 	/** Get Title for Bone Translation Retargeting Mode menu. */
 	FText GetTranslationRetargetingModeMenuTitle() const;
 
-	/** Callback from a slider widget if the text entry or slider movement is used */
+	/** Callback from a slider widget if the text entry is used */
 	void OnBlendSliderCommitted(float NewValue, ETextCommit::Type CommitType);
 
+	/** Callback from a slider widget if the slider is used */
+	void OnBlendSliderChanged(float NewValue);
+
+	/** Callback from a slider widget when the user begins sliding */
+	void OnBeginBlendSliderMovement();
+
+	/** Callback from a slider widget when the user has finished sliding */
+	void OnEndBlendSliderMovement(float NewValue);
+
 	/** Set Translation Retargeting Mode for this bone. */
 	void SetBoneTranslationRetargetingMode(EBoneTranslationRetargetingMode::Type NewRetargetingMode);
 
-	/** Set current Blend Scale for this bone */
-	void SetBoneBlendProfileScale(float NewScale, bool bRecurse);
-
 	/** Get the current Blend Scale for this bone */
 	float GetBoneBlendProfileScale() const;
 
@@ -114,4 +120,6 @@ private:
 	/** Reference to the retargeting combo button */
 	TSharedPtr<SComboButton> RetargetingComboButton;
 
+	/** True if the user is in a transaction when moving the blend profile slider */
+	bool bBlendSliderStartedTransaction;
 };
diff --git a/Engine/Source/Editor/SkeletonEditor/Private/SkeletonTreeVirtualBoneItem.cpp b/Engine/Source/Editor/SkeletonEditor/Private/SkeletonTreeVirtualBoneItem.cpp
--- a/Engine/Source/Editor/SkeletonEditor/Private/SkeletonTreeVirtualBoneItem.cpp
+++ b/Engine/Source/Editor/SkeletonEditor/Private/SkeletonTreeVirtualBoneItem.cpp
@@ -11,6 +11,8 @@
 #include "Animation/DebugSkelMeshComponent.h"
 #include "Animation/BlendProfile.h"
 #include "UObject/Package.h"
+#include "Editor.h"
+#include "ScopedTransaction.h"
 
 #define LOCTEXT_NAMESPACE "FSkeletonTreeVirtualBoneItem"
 
@@ -107,7 +109,10 @@ TSharedRef< SWidget > FSkeletonTreeVirtualBoneItem::GenerateWidgetForDataColumn(
 			.MaxSliderValue(this, &FSkeletonTreeVirtualBoneItem::GetBlendProfileMaxSliderValue)
 			.Value(this, &FSkeletonTreeVirtualBoneItem::GetBoneBlendProfileScale)
 			.OnValueCommitted(this, &FSkeletonTreeVirtualBoneItem::OnBlendSliderCommitted)
-			.OnValueChanged(this, &FSkeletonTreeVirtualBoneItem::OnBlendSliderCommitted, ETextCommit::OnEnter)
+			.OnValueChanged(this, &FSkeletonTreeVirtualBoneItem::OnBlendSliderChanged)
+			.OnBeginSliderMovement(this, &FSkeletonTreeVirtualBoneItem::OnBeginBlendSliderMovement)
+			.OnEndSliderMovement(this, &FSkeletonTreeVirtualBoneItem::OnEndBlendSliderMovement)
+			.ClearKeyboardFocusOnCommit(true)
 			];
 	}
 
@@ -149,15 +154,57 @@ TOptional<float> FSkeletonTreeVirtualBoneItem::GetBlendProfileMinSliderValue() c
 	return 0.0f;
 }
 
-void FSkeletonTreeVirtualBoneItem::OnBlendSliderCommitted(float NewValue, ETextCommit::Type CommitType)
+void FSkeletonTreeVirtualBoneItem::OnBeginBlendSliderMovement()
 {
-	SetBoneBlendProfileScale(NewValue, false);
+	if (bBlendSliderStartedTransaction == false)
+	{
+		bBlendSliderStartedTransaction = true;
+		GEditor->BeginTransaction(LOCTEXT("BlendSliderTransation", "Set Blend Profile Value"));
+
+		const FName& BlendProfileName = GetSkeletonTree()->GetSelectedBlendProfile()->GetFName();
+		UBlendProfile* BlendProfile = GetEditableSkeleton()->GetBlendProfile(BlendProfileName);
+
+		if (BlendProfile)
+		{
+			BlendProfile->SetFlags(RF_Transactional);
+			BlendProfile->Modify();
+		}
+	}
 }
 
-void FSkeletonTreeVirtualBoneItem::SetBoneBlendProfileScale(float NewScale, bool bRecurse)
+void FSkeletonTreeVirtualBoneItem::OnEndBlendSliderMovement(float NewValue)
+{
+	if (bBlendSliderStartedTransaction)
+	{
+		GEditor->EndTransaction();
+		bBlendSliderStartedTransaction = false;
+	}
+}
+
+void FSkeletonTreeVirtualBoneItem::OnBlendSliderCommitted(float NewValue, ETextCommit::Type CommitType)
 {
 	FName BlendProfileName = GetSkeletonTree()->GetSelectedBlendProfile()->GetFName();
-	GetEditableSkeleton()->SetBlendProfileScale(BlendProfileName, BoneName, NewScale, bRecurse);
+	UBlendProfile* BlendProfile = GetEditableSkeleton()->GetBlendProfile(BlendProfileName);
+
+	if (BlendProfile)
+	{
+		FScopedTransaction(LOCTEXT("SetBlendProfileValue", "Set Blend Profile Value"));
+		BlendProfile->SetFlags(RF_Transactional);
+		BlendProfile->Modify();
+
+		BlendProfile->SetBoneBlendScale(BoneName, NewValue, false, true);
+	}
+}
+
+void FSkeletonTreeVirtualBoneItem::OnBlendSliderChanged(float NewValue)
+{
+	const FName& BlendProfileName = GetSkeletonTree()->GetSelectedBlendProfile()->GetFName();
+	UBlendProfile* BlendProfile = GetEditableSkeleton()->GetBlendProfile(BlendProfileName);
+
+	if (BlendProfile)
+	{
+		BlendProfile->SetBoneBlendScale(BoneName, NewValue, false, true);
+	}
 }
 
 FSlateFontInfo FSkeletonTreeVirtualBoneItem::GetBoneTextFont() const
diff --git a/Engine/Source/Editor/SkeletonEditor/Private/SkeletonTreeVirtualBoneItem.h b/Engine/Source/Editor/SkeletonEditor/Private/SkeletonTreeVirtualBoneItem.h
--- a/Engine/Source/Editor/SkeletonEditor/Private/SkeletonTreeVirtualBoneItem.h
+++ b/Engine/Source/Editor/SkeletonEditor/Private/SkeletonTreeVirtualBoneItem.h
@@ -48,15 +48,21 @@ public:
 
 private:
 	
-	/** Callback from a slider widget if the text entry or slider movement is used */
+	/** Callback from a slider widget if the text entry is used */
 	void OnBlendSliderCommitted(float NewValue, ETextCommit::Type CommitType);
 
+	/** Callback from a slider widget if the slider is used */
+	void OnBlendSliderChanged(float NewValue);
+
+	/** Callback from a slider widget when the user begins sliding */
+	void OnBeginBlendSliderMovement();
+
+	/** Callback from a slider widget when the user has finished sliding */
+	void OnEndBlendSliderMovement(float NewValue);
+
 	/** Set Translation Retargeting Mode for this bone. */
 	void SetBoneTranslationRetargetingMode(EBoneTranslationRetargetingMode::Type NewRetargetingMode);
 
-	/** Set current Blend Scale for this bone */
-	void SetBoneBlendProfileScale(float NewScale, bool bRecurse);
-
 	/** Get the current Blend Scale for this bone */
 	float GetBoneBlendProfileScale() const;
 
@@ -109,4 +115,7 @@ private:
 
 	/** Bone proxy used for debug display */
 	UBoneProxy* BoneProxy;
+
+	/** True if the user is in a transaction when moving the blend profile slider */
+	bool bBlendSliderStartedTransaction;
 };
diff --git a/Engine/Source/Runtime/Engine/Private/Animation/BlendProfile.cpp b/Engine/Source/Runtime/Engine/Private/Animation/BlendProfile.cpp
--- a/Engine/Source/Runtime/Engine/Private/Animation/BlendProfile.cpp
+++ b/Engine/Source/Runtime/Engine/Private/Animation/BlendProfile.cpp
@@ -172,7 +172,6 @@ void UBlendProfile::SetSingleBoneBlendScale(int32 InBoneIdx, float InScale, bool
 
 	if(Entry)
 	{
-		Modify();
 		Entry->BlendScale = InScale;
 
 		// Remove any entry that gets set back to DefautBlendScale - so we only store entries that actually contain a scale
