Truncation fixes for SkeletonEditor module

diff --git a/Engine/Source/Editor/SkeletonEditor/Private/BoneProxyDetailsCustomization.cpp b/Engine/Source/Editor/SkeletonEditor/Private/BoneProxyDetailsCustomization.cpp
--- a/Engine/Source/Editor/SkeletonEditor/Private/BoneProxyDetailsCustomization.cpp
+++ b/Engine/Source/Editor/SkeletonEditor/Private/BoneProxyDetailsCustomization.cpp
@@ -112,7 +112,7 @@ static TSharedRef<SWidget> BuildTransformFieldLabel(bool* bValuePtr, const FText
 			.HAlign(HAlign_Left)
 			[
 				SNew(SComboButton)
-				.ContentPadding( 0 )
+				.ContentPadding( 0.f )
 				.ButtonStyle( FAppStyle::Get(), "NoBorder" )
 				.ForegroundColor( FSlateColor::UseForeground() )
 				.MenuContent()
diff --git a/Engine/Source/Editor/SkeletonEditor/Private/SSkeletonTree.cpp b/Engine/Source/Editor/SkeletonEditor/Private/SSkeletonTree.cpp
--- a/Engine/Source/Editor/SkeletonEditor/Private/SSkeletonTree.cpp
+++ b/Engine/Source/Editor/SkeletonEditor/Private/SSkeletonTree.cpp
@@ -614,7 +614,7 @@ void SSkeletonTree::CreateTreeColumns()
 	.HeaderContent()
 	[
 		SNew(SBox)
-		.HeightOverride(24)
+		.HeightOverride(24.f)
 		.HAlign(HAlign_Left)
 		[
 			SAssignNew(BlendProfileHeader, SInlineEditableTextBlock)
@@ -1366,8 +1366,8 @@ void SSkeletonTree::FillAttachAssetSubmenu(FMenuBuilder& MenuBuilder, const TSha
 	AssetPickerConfig.OnAssetSelected = FOnAssetSelected::CreateSP(this, &SSkeletonTree::OnAssetSelectedFromPicker, TargetItem);
 
 	TSharedRef<SWidget> MenuContent = SNew(SBox)
-	.WidthOverride(384)
-	.HeightOverride(500)
+	.WidthOverride(384.f)
+	.HeightOverride(500.f)
 	[
 		ContentBrowserModule.Get().CreateAssetPicker(AssetPickerConfig)
 		];
diff --git a/Engine/Source/Editor/SkeletonEditor/Private/SkeletonTreeBoneItem.cpp b/Engine/Source/Editor/SkeletonEditor/Private/SkeletonTreeBoneItem.cpp
--- a/Engine/Source/Editor/SkeletonEditor/Private/SkeletonTreeBoneItem.cpp
+++ b/Engine/Source/Editor/SkeletonEditor/Private/SkeletonTreeBoneItem.cpp
@@ -108,7 +108,7 @@ TSharedRef< SWidget > FSkeletonTreeBoneItem::GenerateWidgetForDataColumn(const F
 				SAssignNew(RetargetingComboButton, SComboButton)
             	.ComboButtonStyle( &FAppStyle::Get().GetWidgetStyle< FComboButtonStyle >("SkeletonTree.RetargetingComboButton"))
 				.ForegroundColor(this, &FSkeletonTreeBoneItem::GetBoneTextColor, InIsSelected)
-				.ContentPadding(0)
+				.ContentPadding(0.f)
 				.OnGetMenuContent(this, &FSkeletonTreeBoneItem::CreateBoneTranslationRetargetingModeMenu)
 				.ToolTip(IDocumentation::Get()->CreateToolTip(
 					LOCTEXT("RetargetingToolTip", "Set bone translation retargeting mode"),
@@ -512,7 +512,7 @@ bool FSkeletonTreeBoneItem::IsBoneWeighted(int32 MeshBoneIndex, UDebugSkelMeshCo
 	FSkeletalMeshLODRenderData& LODData = PreviewComponent->GetSkeletalMeshAsset()->GetResourceForRendering()->LODRenderData[LODIndex];
 
 	//Check whether the bone is vertex weighted
-	int32 Index = LODData.ActiveBoneIndices.Find(MeshBoneIndex);
+	int32 Index = LODData.ActiveBoneIndices.Find(IntCastChecked<FBoneIndexType>(MeshBoneIndex));
 
 	return Index != INDEX_NONE;
 }
@@ -532,7 +532,7 @@ bool FSkeletonTreeBoneItem::IsBoneRequired(int32 MeshBoneIndex, UDebugSkelMeshCo
 	FSkeletalMeshLODRenderData& LODData = PreviewComponent->GetSkeletalMeshAsset()->GetResourceForRendering()->LODRenderData[LODIndex];
 
 	//Check whether the bone is vertex weighted
-	int32 Index = LODData.RequiredBones.Find(MeshBoneIndex);
+	int32 Index = LODData.RequiredBones.Find(IntCastChecked<FBoneIndexType>(MeshBoneIndex));
 
 	return Index != INDEX_NONE;
 }
