Animation Curve Runtime & Editor Improvements

diff --git a/Engine/Source/Editor/Persona/Public/PersonaModule.h b/Engine/Source/Editor/Persona/Public/PersonaModule.h
--- a/Engine/Source/Editor/Persona/Public/PersonaModule.h
+++ b/Engine/Source/Editor/Persona/Public/PersonaModule.h
@@ -28,6 +28,7 @@ class IAnimSequenceCurveEditor;
 class IAnimationEditor;
 class IDetailLayoutBuilder;
 class FPreviewSceneDescriptionCustomization;
+struct FAnimAssetFindReplaceConfig;
 
 extern const FName PersonaAppName;
 
@@ -299,6 +300,9 @@ public:
 	/** Create a skeleton curve viewer tab factory */
 	virtual TSharedRef<FWorkflowTabFactory> CreateCurveViewerTabFactory(const TSharedRef<class FWorkflowCentricApplication>& InHostingApp, const TSharedRef<class IEditableSkeleton>& InEditableSkeleton, const TSharedRef<class IPersonaPreviewScene>& InPreviewScene, FOnObjectsSelected InOnObjectsSelected) const;
 
+	/** Create a skeleton curve metadata editor tab factory */
+	virtual TSharedRef<FWorkflowTabFactory> CreateCurveMetadataEditorTabFactory(const TSharedRef<class FWorkflowCentricApplication>& InHostingApp, UObject* InMetadataHost, const TSharedRef<class IPersonaPreviewScene>& InPreviewScene, FOnObjectsSelected InOnObjectsSelected) const;
+
 	/** Create a retarget sources tab factory */
 	virtual TSharedRef<FWorkflowTabFactory> CreateRetargetSourcesTabFactory(const TSharedRef<class FWorkflowCentricApplication>& InHostingApp, const TSharedRef<class IEditableSkeleton>& InEditableSkeleton, const TSharedRef<IPersonaPreviewScene>& InPreviewScene, FSimpleMulticastDelegate& InOnPostUndo) const;
 
@@ -342,6 +346,9 @@ public:
 	/** Create a tab factory for editing montage sections */
 	virtual TSharedRef<FWorkflowTabFactory> CreateAnimMontageSectionsTabFactory(const TSharedRef<class FWorkflowCentricApplication>& InHostingApp, const TSharedRef<IPersonaToolkit>& InPersonaToolkit, FSimpleMulticastDelegate& InOnSectionsChanged) const;
 
+	/** Create a tab factory for finding and replacing in anim data */
+	virtual TSharedRef<FWorkflowTabFactory> CreateAnimAssetFindReplaceTabFactory(const TSharedRef<FWorkflowCentricApplication>& InHostingApp, const FAnimAssetFindReplaceConfig& InConfig) const;
+	
 	/** Create a widget that acts as a document for an animation asset */
 	virtual TSharedRef<SWidget> CreateEditorWidgetForAnimDocument(const TSharedRef<IAnimationEditor>& InHostingApp, UObject* InAnimAsset, const FAnimDocumentArgs& InArgs, FString& OutDocumentLink);
 
@@ -358,8 +365,11 @@ public:
 	/** Import a new asset using the supplied skeleton */
 	virtual void ImportNewAsset(USkeleton* InSkeleton, EFBXImportType DefaultImportType);
 
-	/** Check all animations & skeletal meshes for curve usage */
-	virtual void TestSkeletonCurveNamesForUse(const TSharedRef<IEditableSkeleton>& InEditableSkeleton) const;
+	UE_DEPRECATED(5.3, "Please use TestSkeletonCurveMetaDataForUse")
+	virtual void TestSkeletonCurveNamesForUse(const TSharedRef<IEditableSkeleton>& InEditableSkeleton) const { TestSkeletonCurveMetaDataForUse(InEditableSkeleton); }
+
+	/** Check all animations & skeletal meshes for curve metadata usage */
+	virtual void TestSkeletonCurveMetaDataForUse(const TSharedRef<IEditableSkeleton>& InEditableSkeleton) const;
 
 	/** Apply Compression to list of animations and optionally asks to pick an overrides to the bone compression settings */
 	virtual void ApplyCompression(TArray<TWeakObjectPtr<class UAnimSequence>>& AnimSequences, bool bPickBoneSettingsOverride);
@@ -437,6 +447,12 @@ public:
 	/** Register common tabs */
 	virtual FOnRegisterTabs& OnRegisterTabs() { return OnRegisterTabsDelegate; }
 
+	/** Create a widget that can choose a curve name. Derives available names from the asset registry list of assets that use the specified skeleton. */
+	virtual TSharedRef<SWidget> CreateCurvePicker(const USkeleton* InSkeleton, FOnCurvePicked InOnCurvePicked, FIsCurveNameMarkedForExclusion InIsCurveNameMarkedForExclusion = FIsCurveNameMarkedForExclusion());
+
+	UE_DEPRECATED(5.3, "Please use CreateCurvePicker that takes a const USkeleton*")
+	virtual TSharedRef<SWidget> CreateCurvePicker(TSharedRef<IEditableSkeleton> InEditableSkeleton, FOnCurvePicked InOnCurvePicked, FIsCurveNameMarkedForExclusion InIsCurveNameMarkedForExclusion = FIsCurveNameMarkedForExclusion());
+	
 private:
 	/** When a new anim notify blueprint is created, this will handle post creation work such as adding non-event default nodes */
 	void HandleNewAnimNotifyBlueprintCreated(UBlueprint* InBlueprint);
diff --git a/Engine/Source/Editor/Persona/Public/PersonaTabs.h b/Engine/Source/Editor/Persona/Public/PersonaTabs.h
--- a/Engine/Source/Editor/Persona/Public/PersonaTabs.h
+++ b/Engine/Source/Editor/Persona/Public/PersonaTabs.h
@@ -10,6 +10,7 @@ struct PERSONA_API FPersonaTabs
 	// Selection Details
 	static const FName MorphTargetsID;
 	static const FName AnimCurveViewID;
+	static const FName AnimCurveMetadataEditorID;
 	static const FName SkeletonTreeViewID;
 	// Skeleton Pose manager
 	static const FName RetargetManagerID;
@@ -44,6 +45,7 @@ struct PERSONA_API FPersonaTabs
 	// Advanced Preview Scene
 	static const FName AdvancedPreviewSceneSettingsID;
 	static const FName DetailsID;
+	static const FName FindReplaceID;
 
 private:
 	FPersonaTabs() {}
diff --git a/Engine/Source/Editor/SequenceRecorder/Private/AnimationRecorder.cpp b/Engine/Source/Editor/SequenceRecorder/Private/AnimationRecorder.cpp
--- a/Engine/Source/Editor/SequenceRecorder/Private/AnimationRecorder.cpp
+++ b/Engine/Source/Editor/SequenceRecorder/Private/AnimationRecorder.cpp
@@ -268,8 +268,7 @@ void FAnimationRecorder::StartRecord(USkeletalMeshComponent* Component, UAnimSeq
 
 	RecordedCurves.Reset();
 	RecordedTimes.Empty();
-	UIDToArrayIndexLUT = nullptr;
-
+	
 	USkeleton* AnimSkeleton = AnimationObject->GetSkeleton();
 	// add all frames
 	for (int32 BoneIndex=0; BoneIndex <PreviousSpacesBases.Num(); ++BoneIndex)
@@ -388,101 +387,88 @@ UAnimSequence* FAnimationRecorder::StopRecord(bool bShowMessage)
 		// @todo figure out why removing redundant keys is inconsistent
 
 		// add to real curve data 
-		if (RecordedCurves.Num() == NumKeys && UIDToArrayIndexLUT)
+		if (RecordedCurves.Num() == NumKeys)
 		{
-			StartTime = FPlatformTime::Seconds();
-			USkeleton* SkeletonObj = AnimationObject->GetSkeleton();
-			for (int32 CurveUID = 0; CurveUID < UIDToArrayIndexLUT->Num(); ++CurveUID)
+			// Accumulate active curve names over recording duration
+			TSet<FName> AllCurveNames;
+			for(FBlendedHeapCurve& Curve : RecordedCurves)
 			{
-				const int32 CurveIndex = (*UIDToArrayIndexLUT)[CurveUID];
-
-				if (CurveIndex != MAX_uint16)
+				Curve.ForEachElement([this, &AllCurveNames](const UE::Anim::FCurveElement& InElement)
 				{
 					// Skip curves which type is disabled in the recorder settings
-					if (const FCurveMetaData* CurveMetaData = SkeletonObj->GetCurveMetaData(CurveUID))
+					const bool bMorphTarget = EnumHasAnyFlags(InElement.Flags, UE::Anim::ECurveElementFlags::MorphTarget);
+					const bool bMaterial = EnumHasAnyFlags(InElement.Flags, UE::Anim::ECurveElementFlags::Material);
+					const bool bShouldSkipName = ShouldSkipName(InElement.Name);
+					const bool bAttributeCurve = !bMorphTarget && !bMaterial;
+					const bool bSkipCurve = (bMorphTarget && !bRecordMorphTargets) ||
+											(bAttributeCurve && !bRecordAttributeCurves) ||
+											(bMaterial && !bRecordMaterialCurves) ||
+											bShouldSkipName;
+
+					if(!bSkipCurve)
 					{
-						const bool bMorphTarget = CurveMetaData->Type.bMorphtarget;
-						const bool bMaterialCurve = CurveMetaData->Type.bMaterial;
-						const bool bAttributeCurve = !bMorphTarget && !bMaterialCurve;
-						
-						const FSmartNameMapping* Mapping = SkeletonObj->GetSmartNameContainer(USkeleton::AnimCurveMappingName);
-						FSmartName SmartName;
-
-						bool bShouldSkipName = false;
-						if (Mapping && Mapping->FindSmartNameByUID(CurveUID, SmartName))
-						{
-							bShouldSkipName = ShouldSkipName(SmartName.DisplayName);
-						}
-
-						const bool bSkipCurve = (bMorphTarget && !bRecordMorphTargets) ||
-												(bAttributeCurve && !bRecordAttributeCurves) ||
-												(bMaterialCurve && !bRecordMaterialCurves) ||
-												bShouldSkipName;
-
-						if (bSkipCurve)
-						{
-							UE_LOG(LogAnimation, Log, TEXT("Animation Recorder skipping curve: %s"), *SmartName.DisplayName.ToString());
-							continue;
-						}
+						AllCurveNames.Add(InElement.Name);
 					}
+				});
+			}
 
-					const FFloatCurve* FloatCurveData = nullptr;
+			StartTime = FPlatformTime::Seconds();
+			USkeleton* SkeletonObj = AnimationObject->GetSkeleton();
+			for (FName CurveName : AllCurveNames)
+			{
+				const FFloatCurve* FloatCurveData = nullptr;
 
-					TArray<float> TimesToRecord;
-					TArray<float> ValuesToRecord;
-					TimesToRecord.SetNum(NumKeys);
-					ValuesToRecord.SetNum(NumKeys);
+				TArray<float> TimesToRecord;
+				TArray<float> ValuesToRecord;
+				TimesToRecord.SetNum(NumKeys);
+				ValuesToRecord.SetNum(NumKeys);
 
-					bool bSeenThisCurve = false;
-					int32 WriteIndex = 0;
-					for (int32 KeyIndex = 0; KeyIndex < NumKeys; ++KeyIndex)
+				bool bSeenThisCurve = false;
+				int32 WriteIndex = 0;
+				for (int32 KeyIndex = 0; KeyIndex < NumKeys; ++KeyIndex)
+				{
+					const float TimeToRecord = RecordingRate.AsSeconds(KeyIndex);
+					bool bIsCurveValid = false;
+					float CurCurveValue = RecordedCurves[KeyIndex].Get(CurveName, bIsCurveValid);
+					if(bIsCurveValid)
 					{
-						const float TimeToRecord = RecordingRate.AsSeconds(KeyIndex);
-						if(RecordedCurves[KeyIndex].ValidCurveWeights[CurveIndex])
+						if (!bSeenThisCurve)
 						{
-							float CurCurveValue = RecordedCurves[KeyIndex].CurveWeights[CurveIndex];
-							if (!bSeenThisCurve)
-							{
-								bSeenThisCurve = true;
-
-								// add one and save the cache
-								FSmartName CurveName;
-								if (SkeletonObj->GetSmartNameByUID(USkeleton::AnimCurveMappingName, CurveUID, CurveName))
-								{
-									// give default curve flag for recording 
-									const FAnimationCurveIdentifier CurveId(CurveName, ERawCurveTrackTypes::RCT_Float);
-									Controller.AddCurve(CurveId, AACF_DefaultCurve);
-									FloatCurveData = AnimationObject->GetDataModel()->FindFloatCurve(CurveId);
-								}
-							}
-
-							if (FloatCurveData)
-							{
-								TimesToRecord[WriteIndex] = TimeToRecord;
-								ValuesToRecord[WriteIndex] = CurCurveValue;
-
-								++WriteIndex;
-							}
+							bSeenThisCurve = true;
+
+							// add one and save the cache
+							// give default curve flag for recording 
+							const FAnimationCurveIdentifier CurveId(CurveName, ERawCurveTrackTypes::RCT_Float);
+							Controller.AddCurve(CurveId, AACF_DefaultCurve);
+							FloatCurveData = AnimationObject->GetDataModel()->FindFloatCurve(CurveId);
 						}
-					}
 
-					// Fill all the curve data at once
-					if (FloatCurveData)
-					{
-						TArray<FRichCurveKey> Keys;
-						for (int32 Index = 0; Index < WriteIndex; ++Index)
+						if (FloatCurveData)
 						{
-							FRichCurveKey Key(TimesToRecord[Index], ValuesToRecord[Index]);
-							Key.InterpMode = InterpMode;
-							Key.TangentMode = TangentMode;
-							Keys.Add(Key);
+							TimesToRecord[WriteIndex] = TimeToRecord;
+							ValuesToRecord[WriteIndex] = CurCurveValue;
+
+							++WriteIndex;
 						}
+					}
+				}
 
-						const FAnimationCurveIdentifier CurveId(FloatCurveData->Name, ERawCurveTrackTypes::RCT_Float);
-						Controller.SetCurveKeys(CurveId, Keys);
+				// Fill all the curve data at once
+				if (FloatCurveData)
+				{
+					TArray<FRichCurveKey> Keys;
+					for (int32 Index = 0; Index < WriteIndex; ++Index)
+					{
+						FRichCurveKey Key(TimesToRecord[Index], ValuesToRecord[Index]);
+						Key.InterpMode = InterpMode;
+						Key.TangentMode = TangentMode;
+						Keys.Add(Key);
 					}
+
+					const FAnimationCurveIdentifier CurveId(FloatCurveData->GetName(), ERawCurveTrackTypes::RCT_Float);
+					Controller.SetCurveKeys(CurveId, Keys);
 				}
-			}	
+			}
 
 			ElapsedTime = FPlatformTime::Seconds() - StartTime;
 			UE_LOG(LogAnimation, Log, TEXT("Animation Recorder set keys in %0.02f seconds"), ElapsedTime);
@@ -809,16 +795,7 @@ void FAnimationRecorder::UpdateRecord(USkeletalMeshComponent* Component, float D
 			BlendedComponentToWorld.Blend(PreviousComponentToWorld, Component->GetComponentTransform(), BlendAlpha);
 
 			FBlendedHeapCurve BlendedCurve;
-			if (AnimCurves.CurveWeights.Num() > 0 && PreviousAnimCurves.CurveWeights.Num() == AnimCurves.CurveWeights.Num() && PreviousAnimCurves.IsValid() && AnimCurves.IsValid())
-			{
-				BlendedCurve.Lerp(PreviousAnimCurves, AnimCurves, BlendAlpha);
-			}
-			else
-			{
-				// just override with AnimCurves for this frames, because UID list has changed
-				// which means new curves are added in run-time
-				BlendedCurve = AnimCurves;
-			}
+			BlendedCurve.Lerp(PreviousAnimCurves, AnimCurves, BlendAlpha);
 
 			if (!Record(Component, BlendedComponentToWorld, BlendedSpaceBases, BlendedCurve, FramesRecorded + 1))
 			{
@@ -965,21 +942,9 @@ bool FAnimationRecorder::Record(USkeletalMeshComponent* Component, FTransform co
 		}
 		// each RecordedCurves contains all elements
 		const bool bRecordCurves = bRecordMorphTargets || bRecordAttributeCurves || bRecordMaterialCurves; 
-		if (bRecordCurves && AnimationCurves.CurveWeights.Num() > 0)
+		if (bRecordCurves)
 		{
-			RecordedCurves.Emplace(AnimationCurves.CurveWeights, AnimationCurves.ValidCurveWeights);
-			if (UIDToArrayIndexLUT == nullptr)
-			{
-				UIDToArrayIndexLUT = AnimationCurves.UIDToArrayIndexLUT;
-			}
-			else
-			{
-				ensureAlways(UIDToArrayIndexLUT->Num() == AnimationCurves.UIDToArrayIndexLUT->Num());
-				if (UIDToArrayIndexLUT != AnimationCurves.UIDToArrayIndexLUT)
-				{
-					UIDToArrayIndexLUT = AnimationCurves.UIDToArrayIndexLUT;
-				}
-			}
+			RecordedCurves.Emplace(AnimationCurves);
 		}
 
 		LastFrame = FrameToAdd;
diff --git a/Engine/Source/Editor/SequenceRecorder/Public/AnimationRecorder.h b/Engine/Source/Editor/SequenceRecorder/Public/AnimationRecorder.h
--- a/Engine/Source/Editor/SequenceRecorder/Public/AnimationRecorder.h
+++ b/Engine/Source/Editor/SequenceRecorder/Public/AnimationRecorder.h
@@ -160,22 +160,7 @@ private:
 
 	bool ShouldSkipName(const FName& InName) const;
 
-	// recording curve data 
-	struct FBlendedCurve
-	{
-		template<typename Allocator>
-		FBlendedCurve(TArray<float, Allocator> CW, TBitArray<Allocator> VCW)
-		{
-			CurveWeights = CW;
-			ValidCurveWeights = VCW;
-		}
-
-		TArray<float> CurveWeights;
-		TBitArray<> ValidCurveWeights;
-	};
-
-	TArray<FBlendedCurve> RecordedCurves;
-	TArray<uint16> const * UIDToArrayIndexLUT;
+	TArray<FBlendedHeapCurve> RecordedCurves;
 	TArray<FRawAnimSequenceTrack> RawTracks;
 };
 
diff --git a/Engine/Source/Editor/SkeletalMeshEditor/Private/SkeletalMeshEditor.cpp b/Engine/Source/Editor/SkeletalMeshEditor/Private/SkeletalMeshEditor.cpp
--- a/Engine/Source/Editor/SkeletalMeshEditor/Private/SkeletalMeshEditor.cpp
+++ b/Engine/Source/Editor/SkeletalMeshEditor/Private/SkeletalMeshEditor.cpp
@@ -79,7 +79,8 @@ const FName SkeletalMeshEditorTabs::AdvancedPreviewTab(TEXT("AdvancedPreviewTab"
 const FName SkeletalMeshEditorTabs::MorphTargetsTab("MorphTargetsTab");
 const FName SkeletalMeshEditorTabs::AnimationMappingTab("AnimationMappingWindow");
 const FName SkeletalMeshEditorTabs::ToolboxDetailsTab("ToolBoxDetailsTab");
-
+const FName SkeletalMeshEditorTabs::CurveMetadataTab(TEXT("AnimCurveMetadataEditorTab"));
+const FName SkeletalMeshEditorTabs::FindReplaceTab("FindReplaceTab");
 
 DEFINE_LOG_CATEGORY(LogSkeletalMeshEditor);
 
diff --git a/Engine/Source/Editor/SkeletalMeshEditor/Private/SkeletalMeshEditor.h b/Engine/Source/Editor/SkeletalMeshEditor/Private/SkeletalMeshEditor.h
--- a/Engine/Source/Editor/SkeletalMeshEditor/Private/SkeletalMeshEditor.h
+++ b/Engine/Source/Editor/SkeletalMeshEditor/Private/SkeletalMeshEditor.h
@@ -45,6 +45,8 @@ namespace SkeletalMeshEditorTabs
 	extern const FName MeshDetailsTab;
 	extern const FName AnimationMappingTab;
 	extern const FName ToolboxDetailsTab;
+	extern const FName CurveMetadataTab;
+	extern const FName FindReplaceTab;
 }
 
 class FSkeletalMeshEditor : public ISkeletalMeshEditor, public FGCObject, public FEditorUndoClient, public FTickableEditorObject
@@ -115,14 +117,14 @@ public:
 
 	UObject* HandleGetAsset();
 
+	void HandleObjectsSelected(const TArray<UObject*>& InObjects);
+
 	// Returns the currently hosted toolkit. Can be invalid if no toolkit is being hosted.
 	TSharedPtr<IToolkit> GetHostedToolkit() const { return HostedToolkit; }
 
 private:
-	void HandleObjectsSelected(const TArray<UObject*>& InObjects);
-
 	void HandleObjectSelected(UObject* InObject);
-
+	
 	void HandleSelectionChanged(const TArrayView<TSharedPtr<ISkeletonTreeItem>>& InSelectedItems, ESelectInfo::Type InSelectInfo);
 
 	void HandleReimportMesh(int32 SourceFileIndex = INDEX_NONE);
@@ -134,9 +136,6 @@ private:
 
 	void HandleOnPreviewSceneSettingsCustomized(IDetailLayoutBuilder& DetailBuilder);
 
-	/** Callback for toggling UV drawing in the viewport */
-	void ToggleMeshSectionSelection();
-
 	/** Callback for checking whether the UV drawing is switched on. */
 	bool IsMeshSectionSelectionChecked() const;
 
diff --git a/Engine/Source/Editor/SkeletalMeshEditor/Private/SkeletalMeshEditorMode.cpp b/Engine/Source/Editor/SkeletalMeshEditor/Private/SkeletalMeshEditorMode.cpp
--- a/Engine/Source/Editor/SkeletalMeshEditor/Private/SkeletalMeshEditorMode.cpp
+++ b/Engine/Source/Editor/SkeletalMeshEditor/Private/SkeletalMeshEditorMode.cpp
@@ -1,6 +1,7 @@
 // Copyright Epic Games, Inc. All Rights Reserved.
 
 #include "SkeletalMeshEditorMode.h"
+
 #include "PersonaModule.h"
 #include "SkeletalMeshEditor.h"
 #include "ISkeletonTree.h"
@@ -9,6 +10,7 @@
 #include "SControlRigMappingWindow.h"
 #include "Styling/AppStyle.h"
 #include "Engine/SkeletalMesh.h"
+#include "AnimAssetFindReplace.h"
 
 #define LOCTEXT_NAMESPACE "SkeletalMeshEditorMode"
 
@@ -22,6 +24,8 @@ FSkeletalMeshEditorMode::FSkeletalMeshEditorMode(TSharedRef<FWorkflowCentricAppl
 	ISkeletonEditorModule& SkeletonEditorModule = FModuleManager::LoadModuleChecked<ISkeletonEditorModule>("SkeletonEditor");
 	TabFactories.RegisterFactory(SkeletonEditorModule.CreateSkeletonTreeTabFactory(InHostingApp, InSkeletonTree));
 
+	FOnObjectsSelected OnObjectsSelected = FOnObjectsSelected::CreateSP(&SkeletalMeshEditor.Get(), &FSkeletalMeshEditor::HandleObjectsSelected);
+
 	FPersonaModule& PersonaModule = FModuleManager::LoadModuleChecked<FPersonaModule>("Persona");
 	TabFactories.RegisterFactory(PersonaModule.CreateDetailsTabFactory(InHostingApp, FOnDetailsCreated::CreateSP(&SkeletalMeshEditor.Get(), &FSkeletalMeshEditor::HandleDetailsCreated)));
 
@@ -34,10 +38,11 @@ FSkeletalMeshEditorMode::FSkeletalMeshEditorMode(TSharedRef<FWorkflowCentricAppl
 	TabFactories.RegisterFactory(PersonaModule.CreateAdvancedPreviewSceneTabFactory(InHostingApp, SkeletalMeshEditor->GetPersonaToolkit()->GetPreviewScene()));
 	TabFactories.RegisterFactory(PersonaModule.CreateAssetDetailsTabFactory(InHostingApp, FOnGetAsset::CreateSP(&SkeletalMeshEditor.Get(), &FSkeletalMeshEditor::HandleGetAsset), FOnDetailsCreated::CreateSP(&SkeletalMeshEditor.Get(), &FSkeletalMeshEditor::HandleMeshDetailsCreated)));
 	TabFactories.RegisterFactory(PersonaModule.CreateMorphTargetTabFactory(InHostingApp, SkeletalMeshEditor->GetPersonaToolkit()->GetPreviewScene(), SkeletalMeshEditor->OnPostUndo));
-
+	TabFactories.RegisterFactory(PersonaModule.CreateCurveMetadataEditorTabFactory(InHostingApp, SkeletalMeshEditor->HandleGetAsset(), SkeletalMeshEditor->GetPersonaToolkit()->GetPreviewScene(), OnObjectsSelected));
 	TabFactories.RegisterFactory(CreateMeshControllerMappingTabFactory(InHostingApp, Cast<USkeletalMesh> (SkeletalMeshEditor->HandleGetAsset()), SkeletalMeshEditor->OnPostUndo));
+	TabFactories.RegisterFactory(PersonaModule.CreateAnimAssetFindReplaceTabFactory(InHostingApp, FAnimAssetFindReplaceConfig()));
 
-	TabLayout = FTabManager::NewLayout("Standalone_SkeletalMeshEditor_Layout_v3.4")
+	TabLayout = FTabManager::NewLayout("Standalone_SkeletalMeshEditor_Layout_v3.5")
 		->AddArea
 		(
 			FTabManager::NewPrimaryArea()
@@ -58,21 +63,49 @@ FSkeletalMeshEditorMode::FSkeletalMeshEditorMode(TSharedRef<FWorkflowCentricAppl
 				)
 				->Split
 				(
-					FTabManager::NewStack()
+					FTabManager::NewSplitter()
 					->SetSizeCoefficient(0.6f)
-					->SetHideTabWell(true)
-					->AddTab(SkeletalMeshEditorTabs::ViewportTab, ETabState::OpenedTab)
+					->SetOrientation(Orient_Vertical)
+					->Split
+					(	
+						FTabManager::NewStack()
+						->SetSizeCoefficient(0.65)
+						->SetHideTabWell(true)
+						->AddTab(SkeletalMeshEditorTabs::ViewportTab, ETabState::OpenedTab)
+					)
+					->Split
+					(
+						FTabManager::NewStack()
+						->SetSizeCoefficient(0.5f)
+						->SetHideTabWell(false)
+						->AddTab(SkeletalMeshEditorTabs::FindReplaceTab, ETabState::ClosedTab)
+					)
 				)
 				->Split
 				(
-					FTabManager::NewStack()
+					FTabManager::NewSplitter()
 					->SetSizeCoefficient(0.2f)
-					->SetHideTabWell(false)
-					->AddTab(SkeletalMeshEditorTabs::MorphTargetsTab, ETabState::OpenedTab)
-					->AddTab(SkeletalMeshEditorTabs::DetailsTab, ETabState::ClosedTab)
-	                ->AddTab(SkeletalMeshEditorTabs::ToolboxDetailsTab, ETabState::ClosedTab)
-					->AddTab(SkeletalMeshEditorTabs::AdvancedPreviewTab, ETabState::OpenedTab)
-					->SetForegroundTab(SkeletalMeshEditorTabs::MorphTargetsTab)
+					->SetOrientation(Orient_Vertical)
+					->Split
+					(
+						FTabManager::NewStack()
+						->SetSizeCoefficient(0.5f)
+						->SetHideTabWell(false)
+						->AddTab(SkeletalMeshEditorTabs::DetailsTab, ETabState::OpenedTab)
+						->AddTab(SkeletalMeshEditorTabs::ToolboxDetailsTab, ETabState::ClosedTab)
+						->AddTab(SkeletalMeshEditorTabs::AdvancedPreviewTab, ETabState::OpenedTab)
+						->SetForegroundTab(SkeletalMeshEditorTabs::DetailsTab)
+					)
+					->Split
+					(
+
+						FTabManager::NewStack()
+						->SetSizeCoefficient(0.5f)
+						->SetHideTabWell(false)
+						->AddTab(SkeletalMeshEditorTabs::MorphTargetsTab, ETabState::OpenedTab)
+						->AddTab(SkeletalMeshEditorTabs::CurveMetadataTab, ETabState::OpenedTab)
+						->SetForegroundTab(SkeletalMeshEditorTabs::MorphTargetsTab)
+					)
 				)
 			)
 		);
diff --git a/Engine/Source/Editor/SkeletonEditor/Private/EditableSkeleton.cpp b/Engine/Source/Editor/SkeletonEditor/Private/EditableSkeleton.cpp
--- a/Engine/Source/Editor/SkeletonEditor/Private/EditableSkeleton.cpp
+++ b/Engine/Source/Editor/SkeletonEditor/Private/EditableSkeleton.cpp
@@ -269,108 +269,6 @@ bool FEditableSkeleton::DoesSocketAlreadyExist(const USkeletalMeshSocket* InSock
 	return false;
 }
 
-bool FEditableSkeleton::AddSmartname(const FName& InContainerName, const FName& InNewName, FSmartName& OutSmartName)
-{
-	if (const FSmartNameMapping* NameMapping = Skeleton->GetSmartNameContainer(InContainerName))
-	{
-		const bool bResult = Skeleton->AddSmartNameAndModify(InContainerName, InNewName, OutSmartName);
-		OnSmartNameChanged.Broadcast(OutSmartName.DisplayName);
-		return bResult;
-	}
-
-	return false;
-}
-
-void FEditableSkeleton::RenameSmartname(const FName InContainerName, SmartName::UID_Type InNameUid, const FName InNewName)
-{
-	FSmartName CurveToRename;
-	if (!Skeleton->GetSmartNameByUID(USkeleton::AnimCurveMappingName, InNameUid, CurveToRename))
-	{
-		return; //Could not get existing smart name
-	}
-
-	const FSmartNameMapping* Mapping = Skeleton->GetSmartNameContainer(InContainerName);
-	if (!Mapping || Mapping->Exists(InNewName))
-	{
-		return; // Name already exists
-	}
-
-	FText Title = LOCTEXT("RenameCurveDialogTitle", "Confirm Rename");
-	FText ConfirmMessage = LOCTEXT("RenameCurveMessage", "Renaming a curve will necessitate loading and modifying animations and pose assets that use this curve. This could be a slow process.\n\nContinue?");
-
-	if (FMessageDialog::Open(EAppMsgType::YesNo, ConfirmMessage, &Title) == EAppReturnType::Yes)
-	{
-		TArray<FAssetData> AnimationAssets;
-
-		TArray<FName> Names = { CurveToRename.DisplayName };
-		GetAssetsContainingCurves(InContainerName, Names, AnimationAssets);
-
-		// AnimationAssets now only contains assets that are using the selected curve(s)
-		if (AnimationAssets.Num() > 0)
-		{
-			TArray<UAnimSequence*> SequencesToRecompress;
-			SequencesToRecompress.Reserve(AnimationAssets.Num());
-
-			// Proceed to delete the curves
-			GWarn->BeginSlowTask(FText::Format(LOCTEXT("RenameCurvesTaskDesc", "Renaming curve for skeleton {0}"), FText::FromString(Skeleton->GetName())), true);
-			FScopedTransaction Transaction(LOCTEXT("RenameCurvesTransactionName", "Rename skeleton curve"));
-
-			FAnimationCurveIdentifier CurveId(CurveToRename, ERawCurveTrackTypes::RCT_Float);
-
-			//Make sure skeleton is correct before compression 
-			Skeleton->RenameSmartnameAndModify(InContainerName, InNameUid, InNewName);
-			
-			// Remove curves from animation assets
-			for (FAssetData& Data : AnimationAssets)
-			{
-				UObject* Asset = Data.GetAsset();
-
-				if(UAnimSequenceBase* SequenceBase = Cast<UAnimSequenceBase>(Asset))
-				{
-					SequenceBase->Modify();
-
-					if (SequenceBase->GetDataModel()->FindFloatCurve(CurveId))
-					{
-						IAnimationDataController& Controller = SequenceBase->GetController();
-
-						FAnimationCurveIdentifier NewCurveId(FSmartName(InNewName, InNameUid), ERawCurveTrackTypes::RCT_Float);
-						Controller.RenameCurve(CurveId, NewCurveId);
-						if (UAnimSequence* Seq = Cast<UAnimSequence>(SequenceBase))
-						{
-							SequencesToRecompress.Add(Seq);
-						}
-					}
-				}
-				else if (UPoseAsset* PoseAsset = Cast<UPoseAsset>(Asset))
-				{
-					PoseAsset->Modify();
-
-					PoseAsset->RenameSmartName(CurveToRename.DisplayName, InNewName);
-				}
-			}
-			GWarn->EndSlowTask();
-
-			GWarn->BeginSlowTask(LOCTEXT("RebuildingAnimations", "Rebaking/compressing modified animations"), true);
-
-			// Rebake/compress the animations
-			for (UAnimSequence* Seq : SequencesToRecompress)
-			{
-				GWarn->StatusUpdate(1, 2, FText::Format(LOCTEXT("RebuildingAnimationsStatus", "Rebuilding {0}"), FText::FromString(Seq->GetName())));
-				Seq->CacheDerivedDataForCurrentPlatform();
-			}
-
-			GWarn->EndSlowTask();
-		}
-		else
-		{
-			Skeleton->RenameSmartnameAndModify(InContainerName, InNameUid, InNewName);
-		}
-
-		OnSmartNameChanged.Broadcast(InContainerName);
-	}
-
-}
-
 void FEditableSkeleton::GetAssetsContainingCurves(const FName& InContainerName, const TArray<FName>& InNames, TArray<FAssetData>& OutAssets) const
 {
 	FAssetData SkeletonData(Skeleton);
@@ -437,131 +335,6 @@ void FEditableSkeleton::GetAssetsContainingCurves(const FName& InContainerName,
 	GWarn->EndSlowTask();
 }
 
-void FEditableSkeleton::RemoveSmartnamesAndFixupAnimations(const FName& InContainerName, const TArray<FName>& InNames)
-{
-	FText Title = LOCTEXT("RemoveCurveInitialDialogTitle", "Confirm Remove");
-	FText ConfirmMessage = LOCTEXT("RemoveCurveInitialDialogMessage", "Removing curves will necessitate loading and modifying animations and pose assets that use these curves. This could be a slow process.\n\nContinue?");
-
-	if (FMessageDialog::Open(EAppMsgType::YesNo, ConfirmMessage, &Title) != EAppReturnType::Yes)
-	{
-		return;
-	}
-
-	TArray<FAssetData> AnimationAssets;
-	GetAssetsContainingCurves(InContainerName, InNames, AnimationAssets);
-
-	// AnimationAssets now only contains assets that are using the selected curve(s)
-	if (AnimationAssets.Num() > 0)
-	{
-		FString AssetMessage = LOCTEXT("DeleteCurveMessage", "Deleting curves will:\n\nRemove the curves from Animations and PoseAssets\nRemove poses using that curve name from PoseAssets.\n\nThe following assets will be modified. Continue?\n\n").ToString();
-
-		AnimationAssets.Sort([&](const FAssetData& A, const FAssetData& B)
-		{ 
-			if (A.AssetClassPath == B.AssetClassPath)
-			{
-				return A.AssetName.LexicalLess(B.AssetName);
-			}
-			return A.AssetClassPath.Compare(B.AssetClassPath) < 0;
-		});
-
-		for (FAssetData& Data : AnimationAssets)
-		{
-			AssetMessage += Data.AssetName.ToString() + " (" + Data.AssetClassPath.ToString() + ")\n";
-		}
-
-		FText AssetTitleText = LOCTEXT("DeleteCurveDialogTitle", "Confirm Deletion");
-		FText AssetMessageText = FText::FromString(AssetMessage);
-
-		if (FMessageDialog::Open(EAppMsgType::YesNo, AssetMessageText, &AssetTitleText) == EAppReturnType::Yes)
-		{
-			// Proceed to delete the curves
-			GWarn->BeginSlowTask(FText::Format(LOCTEXT("DeleteCurvesTaskDesc", "Deleting curve from skeleton {0}"), FText::FromString(Skeleton->GetName())), true);
-			FScopedTransaction Transaction(LOCTEXT("DeleteCurvesTransactionName", "Delete skeleton curve"));
-
-			// Remove curves from animation assets
-			for (FAssetData& Data : AnimationAssets)
-			{
-				UObject* Asset = Data.GetAsset();
-
-				if(UAnimSequenceBase* Sequence = Cast<UAnimSequenceBase>(Asset))
-				{
-					USkeleton* MySkeleton = Sequence->GetSkeleton();
-					Sequence->Modify(true);
-
-					IAnimationDataController& Controller = Sequence->GetController();
-					IAnimationDataController::FScopedBracket ScopedBracket(Controller, LOCTEXT("DeleteCurvesNotOnSkeletonBracket", "Deleting curves that no longer exist on Skeleton"));
-					for (FName Name : InNames)
-					{
-						FSmartName CurveToDelete;
-						if (MySkeleton->GetSmartNameByName(USkeleton::AnimCurveMappingName, Name, CurveToDelete))
-						{
-							const FAnimationCurveIdentifier CurveId(CurveToDelete, ERawCurveTrackTypes::RCT_Float);
-							Controller.RemoveCurve(CurveId);
-						}
-					}
-				}
-				else if (UPoseAsset* PoseAsset = Cast<UPoseAsset>(Asset))
-				{
-					PoseAsset->Modify();
-					PoseAsset->RemoveSmartNames(InNames);
-				}
-			}
-			GWarn->EndSlowTask();
-
-			// Remove names from skeleton
-			Skeleton->RemoveSmartnamesAndModify(InContainerName, InNames);
-		}
-	}
-	else if(InNames.Num() > 0)
-	{
-		FScopedTransaction Transaction(LOCTEXT("DeleteCurvesTransactionName", "Delete skeleton curve"));
-
-		// Remove names from skeleton
-		Skeleton->RemoveSmartnamesAndModify(InContainerName, InNames);	
-	}
-
-	OnSmartNameChanged.Broadcast(InContainerName);
-}
-
-void FEditableSkeleton::SetCurveMetaDataMaterial(const FSmartName& CurveName, bool bOverrideMaterial)
-{
-	Skeleton->Modify();
-	FCurveMetaData* CurveMetaData = Skeleton->GetCurveMetaData(CurveName);
-	if (CurveMetaData)
-	{
-		// override curve data
-		CurveMetaData->Type.bMaterial = !!bOverrideMaterial;
-	}
-}
-
-void FEditableSkeleton::SetCurveMetaDataMorphTarget(const FSmartName& CurveName, bool bOverrideMorphTarget)
-{
-	Skeleton->Modify();
-	FCurveMetaData* CurveMetaData = Skeleton->GetCurveMetaData(CurveName);
-	if (CurveMetaData)
-	{
-		// override curve data
-		CurveMetaData->Type.bMorphtarget = !!bOverrideMorphTarget;
-	}
-}
-
-void FEditableSkeleton::SetCurveMetaBoneLinks(const FSmartName& CurveName, const TArray<FBoneReference>& BoneLinks, uint8 InMaxLOD)
-{
-	Skeleton->Modify();
-	FCurveMetaData* CurveMetaData = Skeleton->GetCurveMetaData(CurveName);
-	if (CurveMetaData)
-	{
-		// override curve data
-		CurveMetaData->LinkedBones = BoneLinks;
-		CurveMetaData->MaxLOD = InMaxLOD;
-		//  initialize to this skeleton
-		for (FBoneReference& BoneReference : CurveMetaData->LinkedBones)
-		{
-			BoneReference.Initialize(Skeleton);
-		}
-	}
-}
-
 FName FEditableSkeleton::GenerateUniqueSocketName( FName InName, USkeletalMesh* InSkeletalMesh )
 {
 	if(DoesSocketAlreadyExist(nullptr, FText::FromName(InName), ESocketParentType::Skeleton, InSkeletalMesh ) || DoesSocketAlreadyExist(nullptr, FText::FromName(InName), ESocketParentType::Mesh, InSkeletalMesh))
@@ -1741,15 +1514,6 @@ void FEditableSkeleton::RenameSlotName(const FName InOldSlotName, const FName In
 	OnSlotsChanged.Broadcast();
 }
 
-FDelegateHandle FEditableSkeleton::RegisterOnSmartNameChanged(const FOnSmartNameChanged::FDelegate& InOnSmartNameChanged)
-{
-	return OnSmartNameChanged.Add(InOnSmartNameChanged);
-}
-
-void FEditableSkeleton::UnregisterOnSmartNameChanged(FDelegateHandle InHandle)
-{
-	OnSmartNameChanged.Remove(InHandle);
-}
 
 FDelegateHandle FEditableSkeleton::RegisterOnSlotsChanged(const FSimpleMulticastDelegate::FDelegate& InOnSlotsChanged)
 {
diff --git a/Engine/Source/Editor/SkeletonEditor/Private/EditableSkeleton.h b/Engine/Source/Editor/SkeletonEditor/Private/EditableSkeleton.h
--- a/Engine/Source/Editor/SkeletonEditor/Private/EditableSkeleton.h
+++ b/Engine/Source/Editor/SkeletonEditor/Private/EditableSkeleton.h
@@ -50,12 +50,6 @@ public:
 	virtual bool DoesSocketAlreadyExist(const class USkeletalMeshSocket* InSocket, const FText& InSocketName, ESocketParentType SocketParentType, USkeletalMesh* InSkeletalMesh) const override;
 	virtual bool DoesVirtualBoneAlreadyExist(const FString& InVBName) const override;
 	virtual void RenameVirtualBone(const FName OriginalName, const FName InVBName) override;
-	virtual bool AddSmartname(const FName& InContainerName, const FName& InNewName, FSmartName& OutSmartName) override;
-	virtual void RenameSmartname(const FName InContainerName, SmartName::UID_Type InNameUid, const FName InNewName) override;
-	virtual void RemoveSmartnamesAndFixupAnimations(const FName& InContainerName, const TArray<FName>& InNames) override;
-	virtual void SetCurveMetaDataMaterial(const FSmartName& CurveName, bool bOverrideMaterial) override;
-	virtual void SetCurveMetaDataMorphTarget(const FSmartName& CurveName, bool bOverrideMorphTarget) override;
-	virtual void SetCurveMetaBoneLinks(const FSmartName& CurveName, const TArray<FBoneReference>& BoneLinks, uint8 InMaxLOD) override;
 	virtual void SetPreviewMesh(class USkeletalMesh* InSkeletalMesh) override;
 	virtual void LoadAdditionalPreviewSkeletalMeshes() override;
 	virtual void SetAdditionalPreviewSkeletalMeshes(class UDataAsset* InPreviewCollectionAsset) override;
@@ -77,8 +71,6 @@ public:
 	virtual void DeleteSlotName(const FName& InSlotName) override;
 	virtual void DeleteSlotGroup(const FName& InGroupName) override;
 	virtual void RenameSlotName(const FName InOldSlotName, const FName InNewSlotName) override;
-	virtual FDelegateHandle RegisterOnSmartNameChanged(const FOnSmartNameChanged::FDelegate& InOnSmartNameChanged) override;
-	virtual void UnregisterOnSmartNameChanged(FDelegateHandle InHandle) override;
 	virtual void RegisterOnNotifiesChanged(const FSimpleMulticastDelegate::FDelegate& InDelegate) override;
 	virtual void UnregisterOnNotifiesChanged(void* Thing) override;
 	virtual FDelegateHandle RegisterOnSlotsChanged(const FSimpleMulticastDelegate::FDelegate& InOnSlotsChanged) override;
@@ -175,9 +167,6 @@ private:
 	/** Delegate called when trees need refreshing */
 	FSimpleMulticastDelegate OnTreeRefresh;
 
-	/** Delegate called when a smart name is changed */
-	FOnSmartNameChanged OnSmartNameChanged;
-
 	/** Delegate called when notifies are modified */
 	FSimpleMulticastDelegate OnNotifiesChanged;
 
diff --git a/Engine/Source/Editor/SkeletonEditor/Private/SkeletonEditor.cpp b/Engine/Source/Editor/SkeletonEditor/Private/SkeletonEditor.cpp
--- a/Engine/Source/Editor/SkeletonEditor/Private/SkeletonEditor.cpp
+++ b/Engine/Source/Editor/SkeletonEditor/Private/SkeletonEditor.cpp
@@ -38,10 +38,11 @@ const FName SkeletonEditorTabs::SkeletonTreeTab(TEXT("SkeletonTreeView"));
 const FName SkeletonEditorTabs::ViewportTab(TEXT("Viewport"));
 const FName SkeletonEditorTabs::AssetBrowserTab(TEXT("SequenceBrowser"));
 const FName SkeletonEditorTabs::AnimNotifiesTab(TEXT("SkeletonAnimNotifies"));
-const FName SkeletonEditorTabs::CurveNamesTab(TEXT("AnimCurveViewerTab"));
+const FName SkeletonEditorTabs::CurveMetadataTab(TEXT("AnimCurveMetadataEditorTab"));
 const FName SkeletonEditorTabs::AdvancedPreviewTab(TEXT("AdvancedPreviewTab"));
 const FName SkeletonEditorTabs::RetargetManagerTab(TEXT("RetargetManager"));
 const FName SkeletonEditorTabs::SlotNamesTab("SkeletonSlotNames");
+const FName SkeletonEditorTabs::FindReplaceTab("FindReplaceTab");
 
 DEFINE_LOG_CATEGORY(LogSkeletonEditor);
 
@@ -169,8 +170,8 @@ void FSkeletonEditor::BindCommands()
 		FExecuteAction::CreateSP(this, &FSkeletonEditor::RemoveUnusedBones),
 		FCanExecuteAction::CreateSP(this, &FSkeletonEditor::CanRemoveBones));
 
-	ToolkitCommands->MapAction(FSkeletonEditorCommands::Get().TestSkeletonCurveNamesForUse,
-		FExecuteAction::CreateSP(this, &FSkeletonEditor::TestSkeletonCurveNamesForUse));
+	ToolkitCommands->MapAction(FSkeletonEditorCommands::Get().TestSkeletonCurveMetaDataForUse,
+		FExecuteAction::CreateSP(this, &FSkeletonEditor::TestSkeletonCurveMetaDataForUse));
 
 	ToolkitCommands->MapAction(FSkeletonEditorCommands::Get().UpdateSkeletonRefPose,
 		FExecuteAction::CreateSP(this, &FSkeletonEditor::UpdateSkeletonRefPose));
@@ -281,7 +282,7 @@ void FSkeletonEditor::ExtendMenu()
 			{
 				MenuBuilder.AddMenuEntry(FSkeletonEditorCommands::Get().RemoveUnusedBones);
 				MenuBuilder.AddMenuEntry(FSkeletonEditorCommands::Get().UpdateSkeletonRefPose);
-				MenuBuilder.AddMenuEntry(FSkeletonEditorCommands::Get().TestSkeletonCurveNamesForUse);
+				MenuBuilder.AddMenuEntry(FSkeletonEditorCommands::Get().TestSkeletonCurveMetaDataForUse);
 			}
 			MenuBuilder.EndSection();
 		}
@@ -362,10 +363,10 @@ void FSkeletonEditor::RemoveUnusedBones()
 	GetSkeletonTree()->GetEditableSkeleton()->RemoveUnusedBones();
 }
 
-void FSkeletonEditor::TestSkeletonCurveNamesForUse() const
+void FSkeletonEditor::TestSkeletonCurveMetaDataForUse() const
 {
 	FPersonaModule& PersonaModule = FModuleManager::LoadModuleChecked<FPersonaModule>("Persona");
-	PersonaModule.TestSkeletonCurveNamesForUse(GetSkeletonTree()->GetEditableSkeleton());
+	PersonaModule.TestSkeletonCurveMetaDataForUse(GetSkeletonTree()->GetEditableSkeleton());
 }
 
 void FSkeletonEditor::UpdateSkeletonRefPose()
diff --git a/Engine/Source/Editor/SkeletonEditor/Private/SkeletonEditor.h b/Engine/Source/Editor/SkeletonEditor/Private/SkeletonEditor.h
--- a/Engine/Source/Editor/SkeletonEditor/Private/SkeletonEditor.h
+++ b/Engine/Source/Editor/SkeletonEditor/Private/SkeletonEditor.h
@@ -35,10 +35,11 @@ namespace SkeletonEditorTabs
 	extern const FName ViewportTab;
 	extern const FName AssetBrowserTab;
 	extern const FName AnimNotifiesTab;
-	extern const FName CurveNamesTab;
+	extern const FName CurveMetadataTab;
 	extern const FName AdvancedPreviewTab;
 	extern const FName RetargetManagerTab;
 	extern const FName SlotNamesTab;
+	extern const FName FindReplaceTab;
 }
 
 class FSkeletonEditor : public ISkeletonEditor, public FEditorUndoClient, public FTickableEditorObject
@@ -109,7 +110,7 @@ private:
 
 	void RemoveUnusedBones();
 
-	void TestSkeletonCurveNamesForUse() const;
+	void TestSkeletonCurveMetaDataForUse() const;
 
 	void UpdateSkeletonRefPose();
 
diff --git a/Engine/Source/Editor/SkeletonEditor/Private/SkeletonEditorCommands.cpp b/Engine/Source/Editor/SkeletonEditor/Private/SkeletonEditorCommands.cpp
--- a/Engine/Source/Editor/SkeletonEditor/Private/SkeletonEditorCommands.cpp
+++ b/Engine/Source/Editor/SkeletonEditor/Private/SkeletonEditorCommands.cpp
@@ -13,7 +13,7 @@ void FSkeletonEditorCommands::RegisterCommands()
 	// skeleton menu
 	UI_COMMAND(ChangeSkeletonPreviewMesh, "Set Preview Mesh as Default", "Changes the skeletons default preview mesh to the current open preview mesh. The skeleton will require saving after this action.", EUserInterfaceActionType::Button, FInputChord());
 	UI_COMMAND(RemoveUnusedBones, "Remove Unused Bones from Skeleton", "Removes any bones from the skeleton that are not used by any of its meshes. The skeleton and associated animations will require saving after this action.", EUserInterfaceActionType::Button, FInputChord());
-	UI_COMMAND(TestSkeletonCurveNamesForUse, "Test Skeleton Curve Names for Use", "Iterates over all Animations and Meshes for this skeleton to look for unused morph/material curves", EUserInterfaceActionType::Button, FInputChord());
+	UI_COMMAND(TestSkeletonCurveMetaDataForUse, "Test Skeleton Curve Metadata for Use", "Iterates over all Animations and Meshes for this skeleton to look for unused morph/material curves", EUserInterfaceActionType::Button, FInputChord());
 	UI_COMMAND(AnimNotifyWindow, "Anim Notifies", "You can manage animation notifies that belong to the skeleton.", EUserInterfaceActionType::Button, FInputChord());
 	UI_COMMAND(RetargetManager, "Retarget Manager", "Manager retarget setups. ", EUserInterfaceActionType::Button, FInputChord());
 	UI_COMMAND(ImportMesh, "Import Mesh", "Import new mesh for this skeleton. ", EUserInterfaceActionType::Button, FInputChord());
diff --git a/Engine/Source/Editor/SkeletonEditor/Private/SkeletonEditorCommands.h b/Engine/Source/Editor/SkeletonEditor/Private/SkeletonEditorCommands.h
--- a/Engine/Source/Editor/SkeletonEditor/Private/SkeletonEditorCommands.h
+++ b/Engine/Source/Editor/SkeletonEditor/Private/SkeletonEditorCommands.h
@@ -30,7 +30,7 @@ public:
 	// Command to allow users to remove unused bones (not referenced by any skeletalmesh) from the skeleton
 	TSharedPtr<FUICommandInfo> RemoveUnusedBones;
 	// Command to allow users to look for unused curves (curves that exist in animations that do not drive any morph/materials
-	TSharedPtr<FUICommandInfo> TestSkeletonCurveNamesForUse;
+	TSharedPtr<FUICommandInfo> TestSkeletonCurveMetaDataForUse;
 	// Command to show Anim Notify window
 	TSharedPtr<FUICommandInfo> AnimNotifyWindow;
 	// Command to show Retarget Source Manager
diff --git a/Engine/Source/Editor/SkeletonEditor/Private/SkeletonEditorMode.cpp b/Engine/Source/Editor/SkeletonEditor/Private/SkeletonEditorMode.cpp
--- a/Engine/Source/Editor/SkeletonEditor/Private/SkeletonEditorMode.cpp
+++ b/Engine/Source/Editor/SkeletonEditor/Private/SkeletonEditorMode.cpp
@@ -1,12 +1,15 @@
 // Copyright Epic Games, Inc. All Rights Reserved.
 
 #include "SkeletonEditorMode.h"
+
+#include "AnimAssetFindReplace.h"
 #include "Modules/ModuleManager.h"
 #include "PersonaModule.h"
 #include "SkeletonEditor.h"
 #include "ISkeletonTree.h"
 #include "ISkeletonEditorModule.h"
 #include "IPersonaToolkit.h"
+#include "AnimAssetFindReplace.h"
 
 FSkeletonEditorMode::FSkeletonEditorMode(TSharedRef<FWorkflowCentricApplication> InHostingApp, TSharedRef<ISkeletonTree> InSkeletonTree)
 	: FApplicationMode(SkeletonEditorModes::SkeletonEditorMode)
@@ -31,12 +34,13 @@ FSkeletonEditorMode::FSkeletonEditorMode(TSharedRef<FWorkflowCentricApplication>
 	TabFactories.RegisterFactory(PersonaModule.CreateAnimNotifiesTabFactory(InHostingApp, InSkeletonTree->GetEditableSkeleton(),  OnObjectsSelected));
 	TabFactories.RegisterFactory(PersonaModule.CreateAdvancedPreviewSceneTabFactory(InHostingApp, SkeletonEditor->GetPersonaToolkit()->GetPreviewScene()));
 	TabFactories.RegisterFactory(PersonaModule.CreateRetargetSourcesTabFactory(InHostingApp, InSkeletonTree->GetEditableSkeleton(), SkeletonEditor->GetPersonaToolkit()->GetPreviewScene(), SkeletonEditor->OnPostUndo));
-	TabFactories.RegisterFactory(PersonaModule.CreateCurveViewerTabFactory(InHostingApp, InSkeletonTree->GetEditableSkeleton(), SkeletonEditor->GetPersonaToolkit()->GetPreviewScene(), OnObjectsSelected));
+	TabFactories.RegisterFactory(PersonaModule.CreateCurveMetadataEditorTabFactory(InHostingApp, SkeletonEditor->HandleGetAsset(), SkeletonEditor->GetPersonaToolkit()->GetPreviewScene(), OnObjectsSelected));
 	TabFactories.RegisterFactory(PersonaModule.CreateSkeletonSlotNamesTabFactory(InHostingApp, InSkeletonTree->GetEditableSkeleton(), OnObjectSelected));
 	TabFactories.RegisterFactory(PersonaModule.CreateAssetDetailsTabFactory(InHostingApp, FOnGetAsset::CreateSP(&SkeletonEditor.Get(), &FSkeletonEditor::HandleGetAsset), FOnDetailsCreated()));
 	TabFactories.RegisterFactory(PersonaModule.CreateAnimationAssetBrowserTabFactory(InHostingApp, SkeletonEditor->GetPersonaToolkit(), FOnOpenNewAsset::CreateSP(&SkeletonEditor.Get(), &FSkeletonEditor::HandleOpenNewAsset), FOnAnimationSequenceBrowserCreated::CreateSP(&SkeletonEditor.Get(), &FSkeletonEditor::HandleAnimationSequenceBrowserCreated), true));
+	TabFactories.RegisterFactory(PersonaModule.CreateAnimAssetFindReplaceTabFactory(InHostingApp, FAnimAssetFindReplaceConfig()));
 
-	TabLayout = FTabManager::NewLayout("Standalone_SkeletonEditor_Layout_v1.3")
+	TabLayout = FTabManager::NewLayout("Standalone_SkeletonEditor_Layout_v1.4")
 		->AddArea
 		(
 			FTabManager::NewPrimaryArea()
@@ -56,10 +60,23 @@ FSkeletonEditorMode::FSkeletonEditorMode(TSharedRef<FWorkflowCentricApplication>
 				)
 				->Split
 				(
-					FTabManager::NewStack()
+					FTabManager::NewSplitter()
 					->SetSizeCoefficient(0.6f)
-					->SetHideTabWell(true)
-					->AddTab(SkeletonEditorTabs::ViewportTab, ETabState::OpenedTab)
+					->SetOrientation(Orient_Vertical)
+					->Split
+					(	
+						FTabManager::NewStack()
+						->SetSizeCoefficient(0.5f)
+						->SetHideTabWell(true)
+						->AddTab(SkeletonEditorTabs::ViewportTab, ETabState::OpenedTab)
+					)
+					->Split
+					(
+						FTabManager::NewStack()
+						->SetSizeCoefficient(0.5f)
+						->SetHideTabWell(false)
+						->AddTab(SkeletonEditorTabs::FindReplaceTab, ETabState::ClosedTab)
+					)
 				)
 				->Split
 				(
@@ -82,7 +99,7 @@ FSkeletonEditorMode::FSkeletonEditorMode(TSharedRef<FWorkflowCentricApplication>
 						->SetHideTabWell(false)
 						->AddTab(SkeletonEditorTabs::AnimNotifiesTab, ETabState::OpenedTab)
 						->AddTab(SkeletonEditorTabs::AssetBrowserTab, ETabState::OpenedTab)
-						->AddTab(SkeletonEditorTabs::CurveNamesTab, ETabState::OpenedTab)
+						->AddTab(SkeletonEditorTabs::CurveMetadataTab, ETabState::OpenedTab)
 						->AddTab(SkeletonEditorTabs::SlotNamesTab, ETabState::ClosedTab)
 					)
 				)
diff --git a/Engine/Source/Editor/SkeletonEditor/Public/IEditableSkeleton.h b/Engine/Source/Editor/SkeletonEditor/Public/IEditableSkeleton.h
--- a/Engine/Source/Editor/SkeletonEditor/Public/IEditableSkeleton.h
+++ b/Engine/Source/Editor/SkeletonEditor/Public/IEditableSkeleton.h
@@ -63,23 +63,23 @@ public:
 	/** Function to tell you if a socket name already exists. SocketParentType determines where we will look to see if the socket exists, i.e. mesh or skeleton */
 	virtual bool DoesSocketAlreadyExist(const class USkeletalMeshSocket* InSocket, const FText& InSocketName, ESocketParentType SocketParentType, USkeletalMesh* InSkeletalMesh) const = 0;
 
-	/** Add a new smart name. @return true if name was succesfully added */
-	virtual bool AddSmartname(const FName& InContainerName, const FName& InNewName, FSmartName& OutSmartName) = 0;
+	UE_DEPRECATED(5.3, "Please use IInterface_AnimCurveMetaData::AddCurveMetaData.")
+	virtual bool AddSmartname(const FName& InContainerName, const FName& InNewName, FSmartName& OutSmartName) { return false; }
 
-	/** Rename the specified smart name */
-	virtual void RenameSmartname(const FName InContainerName, SmartName::UID_Type InNameUid, const FName InNewName) = 0;
+	UE_DEPRECATED(5.3, "Please use IInterface_AnimCurveMetaData::RenameCurveMetaData.")
+	virtual void RenameSmartname(const FName InContainerName, SmartName::UID_Type InNameUid, const FName InNewName) {}
 
-	/** Remove all the specified smart names and fixup animations that use them */
-	virtual void RemoveSmartnamesAndFixupAnimations(const FName& InContainerName, const TArray<FName>& InNames) = 0;
+	UE_DEPRECATED(5.3, "Please use IInterface_AnimCurveMetaData::RemoveCurveMetaData.")
+	virtual void RemoveSmartnamesAndFixupAnimations(const FName& InContainerName, const TArray<FName>& InNames) {}
 
-	/** Sets Material Meta Data for the curve */
-	virtual void SetCurveMetaDataMaterial(const FSmartName& CurveName, bool bOverrideMaterial) = 0;
+	UE_DEPRECATED(5.3, "Please use IInterface_AnimCurveMetaData::SetCurveMetaDataMaterial.")
+	virtual void SetCurveMetaDataMaterial(const FSmartName& CurveName, bool bOverrideMaterial) {}
 
-	/** Sets Morph Target Meta Data for the curve */
-	virtual void SetCurveMetaDataMorphTarget(const FSmartName& CurveName, bool bOverrideMorphTarget) = 0;
+	UE_DEPRECATED(5.3, "Please use IInterface_AnimCurveMetaData::SetCurveMetaDataMorphTarget.")
+	virtual void SetCurveMetaDataMorphTarget(const FSmartName& CurveName, bool bOverrideMorphTarget) {}
 
-	/** Sets Bone Links per curve */
-	virtual void SetCurveMetaBoneLinks(const FSmartName& CurveName, const TArray<FBoneReference>& BoneLinks, uint8 InMaxLOD) = 0;
+	UE_DEPRECATED(5.3, "Please use IInterface_AnimCurveMetaData::SetCurveMetaDataBoneLinks.")
+	virtual void SetCurveMetaBoneLinks(const FSmartName& CurveName, const TArray<FBoneReference>& BoneLinks, uint8 InMaxLOD) {}
 
 	/**
 	 * Makes sure all attached objects are valid and removes any that aren't.
@@ -197,11 +197,11 @@ public:
 	/** Rename a slot name */
 	virtual void RenameSlotName(const FName InOldSlotName, const FName InNewSlotName) = 0;
 
-	/** Register a delegate to be called when a set of smart names are removed */
-	virtual FDelegateHandle RegisterOnSmartNameChanged(const FOnSmartNameChanged::FDelegate& InOnSmartNameChanged) = 0;
+	UE_DEPRECATED(5.3, "Please use IInterface_AnimCurveMetaData::RegisterOnCurveMetaDataChanged.")
+	virtual FDelegateHandle RegisterOnSmartNameChanged(const FOnSmartNameChanged::FDelegate& InOnSmartNameChanged) { return FDelegateHandle(); }
 
-	/** Register a delegate to be called when a set of smart names are removed */
-	virtual void UnregisterOnSmartNameChanged(FDelegateHandle InHandle) = 0;
+	UE_DEPRECATED(5.3, "Please use IInterface_AnimCurveMetaData::UnegisterOnCurveMetaDataChanged.")
+	virtual void UnregisterOnSmartNameChanged(FDelegateHandle InHandle) {}
 
 	/** Register a delegate to be called when this skeletons notifies are changed */
 	virtual void RegisterOnNotifiesChanged(const FSimpleMulticastDelegate::FDelegate& InDelegate) = 0;
diff --git a/Engine/Source/Editor/UnrealEd/Classes/Animation/EditorAnimCurveBoneLinks.h b/Engine/Source/Editor/UnrealEd/Classes/Animation/EditorAnimCurveBoneLinks.h
--- a/Engine/Source/Editor/UnrealEd/Classes/Animation/EditorAnimCurveBoneLinks.h
+++ b/Engine/Source/Editor/UnrealEd/Classes/Animation/EditorAnimCurveBoneLinks.h
@@ -23,16 +23,22 @@ class UEditorAnimCurveBoneLinks: public UObject
 {
 	GENERATED_UCLASS_BODY()
 public:
-	
+
+	UE_DEPRECATED(5.3, "Please use Initialize that takes a FName")
 	virtual void Initialize(const TWeakPtr<class IEditableSkeleton> InEditableSkeleton, const FSmartName& InCurveName, FOnAnimCurveBonesChange OnChangeIn);
 
+	virtual void Initialize(UAnimCurveMetaData* InAnimCurveMetaData, const FName& InCurveName, FOnAnimCurveBonesChange OnChangeIn);
+
+	UE_DEPRECATED(5.3, "This member is no longer used.")
 	TWeakPtr<class IEditableSkeleton> EditableSkeleton;
+
+	TWeakObjectPtr<UAnimCurveMetaData> AnimCurveMetaData;
 	FOnAnimCurveBonesChange OnChange;
 
-	UPROPERTY(VisibleAnywhere, Category = CurveName)
-	FSmartName CurveName;
+	UPROPERTY(VisibleAnywhere, Category = Curve)
+	FName CurveName;
 
-	UPROPERTY(EditAnywhere, Category = Bones)
+	UPROPERTY(EditAnywhere, Category = Curve)
 	TArray<FBoneReference> ConnectedBones;
 
 	/** Max (Lowest) LOD to evaluate to curve. 
@@ -40,11 +46,14 @@ public:
 	 *  For example, if you have 3 LODs (0, 1, 2), and if you want this to work until LOD 1, type 1.  
 	 *  Then the curve will be evaluated until LOD1, but not for LOD 2
 	 *  Default value is 255 */
-	UPROPERTY(EditAnywhere, Category = LOD)
+	UPROPERTY(EditAnywhere, Category = Curve)
 	uint8 MaxLOD;
 
 	virtual void PostEditChangeProperty(FPropertyChangedEvent& PropertyChangedEvent) override;
 
-	// refresh current Connected Bones data 
+	UE_DEPRECATED(5.3, "Please use Refresh that takes a FName")
 	UNREALED_API void Refresh(const FSmartName& InCurveName, const TArray<FBoneReference>& CurrentLinks, uint8 InMaxLOD);
+
+	// refresh current Connected Bones data 
+	UNREALED_API void Refresh(const FName& InCurveName, const TArray<FBoneReference>& CurrentLinks, uint8 InMaxLOD);
 };
diff --git a/Engine/Source/Editor/UnrealEd/Classes/Factories/FbxAnimSequenceImportData.h b/Engine/Source/Editor/UnrealEd/Classes/Factories/FbxAnimSequenceImportData.h
--- a/Engine/Source/Editor/UnrealEd/Classes/Factories/FbxAnimSequenceImportData.h
+++ b/Engine/Source/Editor/UnrealEd/Classes/Factories/FbxAnimSequenceImportData.h
@@ -93,6 +93,10 @@ class UNREALED_API UFbxAnimSequenceImportData : public UFbxAssetImportData
 	UPROPERTY(EditAnywhere, AdvancedDisplay, config, Category = ImportSettings, meta = (EditCondition = "bImportCustomAttribute", DisplayName="Set Material Curve Type"))
 	bool bSetMaterialDriveParameterOnCustomAttribute;
 
+	/** Whether to automatically add curve metadata to an animation's skeleton. If this is disabled, curve metadata will be added to skeletal meshes for morph targets, but no metadata entry will be created for general curves. */
+	UPROPERTY(EditAnywhere, AdvancedDisplay, config, Category = ImportSettings)
+	bool bAddCurveMetadataToSkeleton;
+
 	/** Set Material Curve Type for the custom attribute with the following suffixes. This doesn't matter if Set Material Curve Type is true  */
 	UPROPERTY(EditAnywhere, AdvancedDisplay, config, Category = ImportSettings, meta = (EditCondition = "bImportCustomAttribute", DisplayName = "Material Curve Suffixes"))
 	TArray<FString> MaterialCurveSuffixes;
diff --git a/Engine/Source/Editor/UnrealEd/Private/Animation/EditorAnimCurveBoneLinks.cpp b/Engine/Source/Editor/UnrealEd/Private/Animation/EditorAnimCurveBoneLinks.cpp
--- a/Engine/Source/Editor/UnrealEd/Private/Animation/EditorAnimCurveBoneLinks.cpp
+++ b/Engine/Source/Editor/UnrealEd/Private/Animation/EditorAnimCurveBoneLinks.cpp
@@ -16,7 +16,15 @@ UEditorAnimCurveBoneLinks::UEditorAnimCurveBoneLinks(const FObjectInitializer& O
 
 void UEditorAnimCurveBoneLinks::Initialize(const TWeakPtr<class IEditableSkeleton> InEditableSkeleton, const FSmartName& InCurveName, FOnAnimCurveBonesChange OnChangeIn)
 {
-	EditableSkeleton = InEditableSkeleton;
+	AnimCurveMetaData = nullptr;
+	CurveName = InCurveName.DisplayName;
+	OnChange = OnChangeIn;
+	SetFlags(RF_Transactional);
+}
+
+void UEditorAnimCurveBoneLinks::Initialize(UAnimCurveMetaData* InAnimCurveMetaData, const FName& InCurveName, FOnAnimCurveBonesChange OnChangeIn)
+{
+	AnimCurveMetaData = InAnimCurveMetaData;
 	CurveName = InCurveName;
 	OnChange = OnChangeIn;
 	SetFlags(RF_Transactional);
@@ -24,7 +32,18 @@ void UEditorAnimCurveBoneLinks::Initialize(const TWeakPtr<class IEditableSkeleto
 
 void UEditorAnimCurveBoneLinks::Refresh(const FSmartName& InCurveName, const TArray<FBoneReference>& CurrentLinks, uint8 InMaxLOD)
 {
-	if (EditableSkeleton.IsValid())
+	if (AnimCurveMetaData.IsValid())
+	{
+		// double check the name
+		CurveName = InCurveName.DisplayName;
+		ConnectedBones = CurrentLinks;
+		MaxLOD = InMaxLOD;
+	}
+}
+
+void UEditorAnimCurveBoneLinks::Refresh(const FName& InCurveName, const TArray<FBoneReference>& CurrentLinks, uint8 InMaxLOD)
+{
+	if (AnimCurveMetaData.IsValid())
 	{
 		// double check the name
 		CurveName = InCurveName;
diff --git a/Engine/Source/Editor/UnrealEd/Private/Factories/MirrorDataTableFactory.cpp b/Engine/Source/Editor/UnrealEd/Private/Factories/MirrorDataTableFactory.cpp
--- a/Engine/Source/Editor/UnrealEd/Private/Factories/MirrorDataTableFactory.cpp
+++ b/Engine/Source/Editor/UnrealEd/Private/Factories/MirrorDataTableFactory.cpp
@@ -343,7 +343,6 @@ UObject* UMirrorDataTableFactory::FactoryCreateNew(UClass* Class, UObject* InPar
 			if (Skeleton)
 			{
 				DataTable->FindReplaceMirroredNames();
-				Skeleton->OnSmartNamesChangedEvent.AddUObject(DataTable, &UMirrorDataTable::FillMirrorArrays);
 			}
 		}
 	}
diff --git a/Engine/Source/Editor/UnrealEd/Private/Factories/PoseAssetFactory.cpp b/Engine/Source/Editor/UnrealEd/Private/Factories/PoseAssetFactory.cpp
--- a/Engine/Source/Editor/UnrealEd/Private/Factories/PoseAssetFactory.cpp
+++ b/Engine/Source/Editor/UnrealEd/Private/Factories/PoseAssetFactory.cpp
@@ -228,22 +228,10 @@ UObject* UPoseAssetFactory::FactoryCreateNew(UClass* Class, UObject* InParent, F
 		TargetSkeleton = SourceAnimation->GetSkeleton();
 
 		UPoseAsset* PoseAsset = NewObject<UPoseAsset>(InParent, Class, Name, Flags);
-		TArray<FSmartName> InputPoseNames;
-		if (PoseNames.Num() > 0)
+		TArray<FName> InputPoseNames;
+		for (int32 Index = 0; Index < PoseNames.Num(); ++Index)
 		{
-			for (int32 Index = 0; Index < PoseNames.Num(); ++Index)
-			{
-				FName PoseName = FName(*PoseNames[Index]);
-				FSmartName NewName;
-				if (TargetSkeleton->GetSmartNameByName(USkeleton::AnimCurveMappingName, PoseName, NewName) == false)
-				{
-					// if failed, add it
-					TargetSkeleton->AddSmartNameAndModify(USkeleton::AnimCurveMappingName, PoseName, NewName);
-				}
-
-				// we want same names in multiple places
-				InputPoseNames.AddUnique(NewName);
-			}
+			InputPoseNames.AddUnique(FName(*PoseNames[Index]));
 		}
 
 		PoseAsset->CreatePoseFromAnimation( SourceAnimation, &InputPoseNames);
@@ -251,7 +239,7 @@ UObject* UPoseAssetFactory::FactoryCreateNew(UClass* Class, UObject* InParent, F
 		return PoseAsset;
 	}
 
-	return NULL;
+	return nullptr;
 }
 
 void UPoseAssetFactory::OnWindowUserActionDelegate(bool bCreate, UAnimSequence* InSequence, const TArray<FString>& InPoseNames)
diff --git a/Engine/Source/Editor/UnrealEd/Private/Fbx/FbxAnimSequenceImportData.cpp b/Engine/Source/Editor/UnrealEd/Private/Fbx/FbxAnimSequenceImportData.cpp
--- a/Engine/Source/Editor/UnrealEd/Private/Fbx/FbxAnimSequenceImportData.cpp
+++ b/Engine/Source/Editor/UnrealEd/Private/Fbx/FbxAnimSequenceImportData.cpp
@@ -8,6 +8,7 @@ UFbxAnimSequenceImportData::UFbxAnimSequenceImportData(const FObjectInitializer&
 	, bImportMeshesInBoneHierarchy(true)
 	, bImportCustomAttribute(true)
 	, bImportBoneTracks(true)
+	, bAddCurveMetadataToSkeleton(true)
 	, bRemoveRedundantKeys(true)
 	, bDoNotImportCurveWithZero(true)
 {
@@ -93,6 +94,7 @@ void UFbxAnimSequenceImportData::CopyAnimationValues(const UFbxAnimSequenceImpor
 	bPreserveLocalTransform = Other->bPreserveLocalTransform;
 	bRemoveRedundantKeys = Other->bRemoveRedundantKeys;
 	bSetMaterialDriveParameterOnCustomAttribute = Other->bSetMaterialDriveParameterOnCustomAttribute;
+	bAddCurveMetadataToSkeleton = Other->bAddCurveMetadataToSkeleton;
 	bUseDefaultSampleRate = Other->bUseDefaultSampleRate;
 	CustomSampleRate = Other->CustomSampleRate;
 	FrameImportRange = Other->FrameImportRange;
diff --git a/Engine/Source/Editor/UnrealEd/Private/Fbx/FbxAnimationExport.cpp b/Engine/Source/Editor/UnrealEd/Private/Fbx/FbxAnimationExport.cpp
--- a/Engine/Source/Editor/UnrealEd/Private/Fbx/FbxAnimationExport.cpp
+++ b/Engine/Source/Editor/UnrealEd/Private/Fbx/FbxAnimationExport.cpp
@@ -92,34 +92,24 @@ void FFbxExporter::ExportAnimSequenceToFbx(const UAnimSequence* AnimSeq, const U
 	}
 
 	//Prepare root anim curves data to be exported
-	TArray<FName> AnimCurveNames;
 	TMap<FName, FbxAnimCurve*> CustomCurveMap;
 	if (BoneNodes.Num() > 0)
 	{
-		const FSmartNameMapping* AnimCurveMapping = Skeleton->GetSmartNameContainer(USkeleton::AnimCurveMappingName);
-		
-		if (AnimCurveMapping)
-		{
-			AnimCurveMapping->FillNameArray(AnimCurveNames);
-
-			const UFbxExportOption* ExportOptions = GetExportOptions();
-			const bool bExportMorphTargetCurvesInMesh = ExportOptions && ExportOptions->bExportPreviewMesh && ExportOptions->bExportMorphTargets;
+		const UFbxExportOption* ExportOptions = GetExportOptions();
+		const bool bExportMorphTargetCurvesInMesh = ExportOptions && ExportOptions->bExportPreviewMesh && ExportOptions->bExportMorphTargets;
 
-			for (auto AnimCurveName : AnimCurveNames)
+		Skeleton->ForEachCurveMetaData([&CustomCurveMap, &BoneNodes, InAnimLayer, bExportMorphTargetCurvesInMesh](const FName& InCurveName, const FCurveMetaData& InMetaData)
+		{
+			//Only export the custom curve if it is not used in a MorphTarget that will be exported latter on.
+			if(!(bExportMorphTargetCurvesInMesh && InMetaData.Type.bMorphtarget))
 			{
-				const FCurveMetaData* CurveMetaData = AnimCurveMapping->GetCurveMetaData(AnimCurveName);
-
-				//Only export the custom curve if it is not used in a MorphTarget that will be exported latter on.
-				if(!(bExportMorphTargetCurvesInMesh && CurveMetaData && CurveMetaData->Type.bMorphtarget))
-				{
-					FbxProperty AnimCurveFbxProp = FbxProperty::Create(BoneNodes[0], FbxDoubleDT, TCHAR_TO_ANSI(*AnimCurveName.ToString()));
-					AnimCurveFbxProp.ModifyFlag(FbxPropertyFlags::eAnimatable, true);
-					AnimCurveFbxProp.ModifyFlag(FbxPropertyFlags::eUserDefined, true);
-					FbxAnimCurve* AnimFbxCurve = AnimCurveFbxProp.GetCurve(InAnimLayer, true);
-					CustomCurveMap.Add(AnimCurveName, AnimFbxCurve);
-				}
+				FbxProperty AnimCurveFbxProp = FbxProperty::Create(BoneNodes[0], FbxDoubleDT, TCHAR_TO_ANSI(*InCurveName.ToString()));
+				AnimCurveFbxProp.ModifyFlag(FbxPropertyFlags::eAnimatable, true);
+				AnimCurveFbxProp.ModifyFlag(FbxPropertyFlags::eUserDefined, true);
+				FbxAnimCurve* AnimFbxCurve = AnimCurveFbxProp.GetCurve(InAnimLayer, true);
+				CustomCurveMap.Add(InCurveName, AnimFbxCurve);
 			}
-		}
+		});
 	}
 
 	ExportCustomAnimCurvesToFbx(CustomCurveMap, AnimSeq, StartFrameTime, EndFrameTime, FrameRateScale, StartTime);
@@ -304,27 +294,13 @@ void FFbxExporter::ExportCustomAnimCurvesToFbx(const TMap<FName, FbxAnimCurve*>&
 {
 	// stack allocator for extracting curve
 	FMemMark Mark(FMemStack::Get());
-	const USkeleton* Skeleton = AnimSeq->GetSkeleton();
-	const FSmartNameMapping* SmartNameMapping = Skeleton ? Skeleton->GetSmartNameContainer(USkeleton::AnimCurveMappingName) : nullptr;
 
-	if (!Skeleton || !SmartNameMapping || !SetupAnimStack(AnimSeq))
+	if (!SetupAnimStack(AnimSeq))
 	{
 		//Something is wrong.
 		return;
 	}
 
-	TArray<SmartName::UID_Type> AnimCurveUIDs;
-	{
-		//We need to recreate the UIDs array manually so that we keep the empty entries otherwise the BlendedCurve won't have the correct mapping.
-		TArray<FName> UID_ToNameArray;
-		SmartNameMapping->FillUIDToNameArray(UID_ToNameArray);
-		AnimCurveUIDs.Reserve(UID_ToNameArray.Num());
-		for (int32 NameIndex = 0; NameIndex < UID_ToNameArray.Num(); ++NameIndex)
-		{
-			AnimCurveUIDs.Add(NameIndex);
-		}
-	}
-
 	for (auto CustomCurve : CustomCurves)
 	{
 		CustomCurve.Value->KeyModifyBegin();
@@ -332,21 +308,14 @@ void FFbxExporter::ExportCustomAnimCurvesToFbx(const TMap<FName, FbxAnimCurve*>&
 	
 	auto ExportLambda = [&](float AnimTime, FbxTime ExportTime, bool bLastKey) {
 		FBlendedCurve BlendedCurve;
-		BlendedCurve.InitFrom(&AnimCurveUIDs);
 		AnimSeq->EvaluateCurveData(BlendedCurve, AnimTime, true);
-		if (BlendedCurve.IsValid())
+		
+		//Loop over the custom curves and add the actual keys
+		for (auto CustomCurve : CustomCurves)
 		{
-			//Loop over the custom curves and add the actual keys
-			for (auto CustomCurve : CustomCurves)
-			{
-				SmartName::UID_Type NameUID = Skeleton->GetUIDByName(USkeleton::AnimCurveMappingName, CustomCurve.Key);
-				if (NameUID != SmartName::MaxUID)
-				{
-					float CurveValueAtTime = BlendedCurve.Get(NameUID) * ValueScale;
-					int32 KeyIndex = CustomCurve.Value->KeyAdd(ExportTime);
-					CustomCurve.Value->KeySetValue(KeyIndex, CurveValueAtTime);
-				}
-			}
+			float CurveValueAtTime = BlendedCurve.Get(CustomCurve.Key) * ValueScale;
+			int32 KeyIndex = CustomCurve.Value->KeyAdd(ExportTime);
+			CustomCurve.Value->KeySetValue(KeyIndex, CurveValueAtTime);
 		}
 	};
 
diff --git a/Engine/Source/Editor/UnrealEd/Private/Fbx/FbxMainImport.cpp b/Engine/Source/Editor/UnrealEd/Private/Fbx/FbxMainImport.cpp
--- a/Engine/Source/Editor/UnrealEd/Private/Fbx/FbxMainImport.cpp
+++ b/Engine/Source/Editor/UnrealEd/Private/Fbx/FbxMainImport.cpp
@@ -501,6 +501,7 @@ void ApplyImportUIToImportOptions(UFbxImportUI* ImportUI, FBXImportOptions& InOu
 		InOutImportOptions.bDeleteExistingNonCurveCustomAttributes = ImportUI->AnimSequenceImportData->bDeleteExistingNonCurveCustomAttributes;
 		InOutImportOptions.bImportBoneTracks			= ImportUI->AnimSequenceImportData->bImportBoneTracks;
 		InOutImportOptions.bSetMaterialDriveParameterOnCustomAttribute = ImportUI->AnimSequenceImportData->bSetMaterialDriveParameterOnCustomAttribute;
+		InOutImportOptions.bAddCurveMetadataToSkeleton	= ImportUI->AnimSequenceImportData->bAddCurveMetadataToSkeleton;
 		InOutImportOptions.MaterialCurveSuffixes		= ImportUI->AnimSequenceImportData->MaterialCurveSuffixes;
 	}
 }
diff --git a/Engine/Source/Editor/UnrealEd/Private/Fbx/FbxSkeletalMeshExport.cpp b/Engine/Source/Editor/UnrealEd/Private/Fbx/FbxSkeletalMeshExport.cpp
--- a/Engine/Source/Editor/UnrealEd/Private/Fbx/FbxSkeletalMeshExport.cpp
+++ b/Engine/Source/Editor/UnrealEd/Private/Fbx/FbxSkeletalMeshExport.cpp
@@ -259,7 +259,6 @@ FbxNode* FFbxExporter::CreateMesh(const USkeletalMesh* SkelMesh, const TCHAR* Me
 
 	if (GetExportOptions()->bExportMorphTargets && SkelMesh->GetSkeleton()) //The skeleton can be null if this is a destructible mesh.
 	{
-		const FSmartNameMapping* SmartNameMapping = SkelMesh->GetSkeleton()->GetSmartNameContainer(USkeleton::AnimCurveMappingName);
 		TMap<FName, FbxAnimCurve*> BlendShapeCurvesMap;
 
 		if (SkelMesh->GetMorphTargets().Num())
@@ -315,7 +314,8 @@ FbxNode* FFbxExporter::CreateMesh(const USkeletalMesh* SkelMesh, const TCHAR* Me
 
 					BlendShapeChannel->AddTargetShape(Shape);
 					FName MorphTargetName = MorphTarget->GetFName();
-					if (AnimSeq && SmartNameMapping && SmartNameMapping->GetCurveMetaData(MorphTargetName) && SmartNameMapping->GetCurveMetaData(MorphTargetName)->Type.bMorphtarget)
+					const FCurveMetaData* CurveMetaData = SkelMesh->GetSkeleton()->GetCurveMetaData(MorphTargetName);
+					if (AnimSeq && CurveMetaData && CurveMetaData->Type.bMorphtarget)
 					{
 						FbxAnimCurve* AnimCurve = Mesh->GetShapeChannel(DeformerIndex, BlendShape->GetBlendShapeChannelCount() - 1, AnimLayer, true);
 						BlendShapeCurvesMap.Add(MorphTargetName, AnimCurve);
diff --git a/Engine/Source/Editor/UnrealEd/Private/Fbx/FbxSkeletalMeshImport.cpp b/Engine/Source/Editor/UnrealEd/Private/Fbx/FbxSkeletalMeshImport.cpp
--- a/Engine/Source/Editor/UnrealEd/Private/Fbx/FbxSkeletalMeshImport.cpp
+++ b/Engine/Source/Editor/UnrealEd/Private/Fbx/FbxSkeletalMeshImport.cpp
@@ -4440,6 +4440,37 @@ void UnFbx::FFbxImporter::ImportMorphTargetsInternal( TArray<FbxNode*>& SkelMesh
 		BaseImportData.MorphTargets.Add(ShapeImportData);
 		check(BaseImportData.MorphTargetNames.Num() == BaseImportData.MorphTargets.Num() && BaseImportData.MorphTargetNames.Num() == BaseImportData.MorphTargetModifiedPoints.Num());
 
+		// Ensure that we have curve metadata for this morph target (either skeleton or mesh)
+		FName CurveName = *ShapeName;
+		if(ImportOptions->bAddCurveMetadataToSkeleton)
+		{
+			if(USkeleton* Skeleton = BaseSkelMesh->GetSkeleton())
+			{
+				Skeleton->AddCurveMetaData(CurveName);
+
+				// Ensure we have a morph flag set
+				FCurveMetaData* CurveMetaData = Skeleton->GetCurveMetaData(CurveName);
+				check(CurveMetaData);
+				CurveMetaData->Type.bMorphtarget = true;
+			}
+		}
+		else
+		{
+			UAnimCurveMetaData* AnimCurveMetaData = BaseSkelMesh->GetAssetUserData<UAnimCurveMetaData>();
+			if(AnimCurveMetaData == nullptr)
+			{
+				AnimCurveMetaData = NewObject<UAnimCurveMetaData>(BaseSkelMesh, NAME_None, RF_Transactional);
+				BaseSkelMesh->AddAssetUserData(AnimCurveMetaData);
+			}
+
+			AnimCurveMetaData->AddCurveMetaData(CurveName);
+
+			// Ensure we have a morph flag set
+			FCurveMetaData* CurveMetaData = AnimCurveMetaData->GetCurveMetaData(CurveName);
+			check(CurveMetaData);
+			CurveMetaData->Type.bMorphtarget = true;
+		}
+
 		if (ImportOptions->bIsImportCancelable && ImportMorphTargetSlowTask.ShouldCancel())
 		{
 			bImportOperationCanceled = true;
diff --git a/Engine/Source/Editor/UnrealEd/Private/SkeletalMeshEdit.cpp b/Engine/Source/Editor/UnrealEd/Private/SkeletalMeshEdit.cpp
--- a/Engine/Source/Editor/UnrealEd/Private/SkeletalMeshEdit.cpp
+++ b/Engine/Source/Editor/UnrealEd/Private/SkeletalMeshEdit.cpp
@@ -1368,20 +1368,8 @@ bool UnFbx::FFbxImporter::ImportCurveToAnimSequence(class UAnimSequence * Target
 	if (TargetSequence && FbxCurve)
 	{
 		FName Name = *CurveName;
-		USkeleton* Skeleton = TargetSequence->GetSkeleton();
-		const FSmartNameMapping* NameMapping = Skeleton->GetSmartNameContainer(USkeleton::AnimCurveMappingName);
 
-		// Add or retrieve curve
-		if (!NameMapping->Exists(Name))
-		{
-			// mark skeleton dirty
-			Skeleton->Modify();
-		}
-
-		FSmartName NewName;
-		Skeleton->AddSmartNameAndModify(USkeleton::AnimCurveMappingName, Name, NewName);
-
-		FAnimationCurveIdentifier FloatCurveId(NewName, ERawCurveTrackTypes::RCT_Float);
+		FAnimationCurveIdentifier FloatCurveId(Name, ERawCurveTrackTypes::RCT_Float);
 
 		const bool bShouldTransact = bReimport;
 		IAnimationDataModel* DataModel = TargetSequence->GetDataModel();
@@ -1403,8 +1391,6 @@ bool UnFbx::FFbxImporter::ImportCurveToAnimSequence(class UAnimSequence * Target
 		// Should be valid at this point
 		ensure(TargetCurve);
 
-		Controller.UpdateCurveNamesFromSkeleton(Skeleton, ERawCurveTrackTypes::RCT_Float, bShouldTransact);
-
 		FRichCurve RichCurve;
 		constexpr bool bNegative = false;
 		if (ImportCurve(FbxCurve, RichCurve, AnimTimeSpan, bNegative, ValueScale))
@@ -1433,21 +1419,8 @@ bool UnFbx::FFbxImporter::ImportRichCurvesToAnimSequence(UAnimSequence* TargetSe
 		for (int32 CurveIndex = 0; CurveIndex < CurveNames.Num(); ++CurveIndex)
 		{
 			FName Name = *CurveNames[CurveIndex];
-			USkeleton* Skeleton = TargetSequence->GetSkeleton();
-			const FSmartNameMapping* NameMapping = Skeleton->GetSmartNameContainer(USkeleton::AnimCurveMappingName);
-
-			// Add or retrieve curve
-			if (!NameMapping->Exists(Name))
-			{
-				// mark skeleton dirty
-				Skeleton->Modify();
-			}
-
-			FSmartName NewName;
-			Skeleton->AddSmartNameAndModify(USkeleton::AnimCurveMappingName, Name, NewName);
-
-			FAnimationCurveIdentifier FloatCurveId(NewName, ERawCurveTrackTypes::RCT_Float);
-
+			
+			FAnimationCurveIdentifier FloatCurveId(Name, ERawCurveTrackTypes::RCT_Float);
 			const IAnimationDataModel* DataModel = TargetSequence->GetDataModel();
 			IAnimationDataController& Controller = TargetSequence->GetController();
 
@@ -1468,8 +1441,6 @@ bool UnFbx::FFbxImporter::ImportRichCurvesToAnimSequence(UAnimSequence* TargetSe
 			// Should be valid at this point
 			ensure(TargetCurve);
 
-			Controller.UpdateCurveNamesFromSkeleton(Skeleton, ERawCurveTrackTypes::RCT_Float, bShouldTransact);
-
 			// Set actual keys on curve within the model
 			Controller.SetCurveKeys(FloatCurveId, RichCurves[CurveIndex].GetConstRefOfKeys(), bShouldTransact);
 
@@ -1786,7 +1757,7 @@ bool UnFbx::FFbxImporter::ImportAnimation(USkeleton* Skeleton, UAnimSequence * D
 		{
 			for (const FFloatCurve& Curve : DestSeq->GetDataModel()->GetFloatCurves())
 			{
-				const FAnimationCurveIdentifier CurveId(Curve.Name, ERawCurveTrackTypes::RCT_Float);
+				const FAnimationCurveIdentifier CurveId(Curve.GetName(), ERawCurveTrackTypes::RCT_Float);
 				Controller.ScaleCurve(CurveId, 0.f, ScaleFactor, bShouldTransact);
 			}
 		}
@@ -1797,21 +1768,21 @@ bool UnFbx::FFbxImporter::ImportAnimation(USkeleton* Skeleton, UAnimSequence * D
 
 	if (ImportOptions->bDeleteExistingMorphTargetCurves || ImportOptions->bDeleteExistingCustomAttributeCurves)
 	{
-		TArray<FSmartName> CurveSmartNamesToRemove;
+		TArray<FName> CurveNamesToRemove;
 		for (const FFloatCurve& Curve : DestSeq->GetDataModel()->GetFloatCurves())
 		{
-			const FCurveMetaData* MetaData = MySkeleton->GetCurveMetaData(Curve.Name);
+			const FCurveMetaData* MetaData = MySkeleton->GetCurveMetaData(Curve.GetName());
 			if (MetaData)
 			{
 				bool bDeleteCurve = MetaData->Type.bMorphtarget ? ImportOptions->bDeleteExistingMorphTargetCurves : ImportOptions->bDeleteExistingCustomAttributeCurves;
 				if (bDeleteCurve)
 				{
-					CurveSmartNamesToRemove.Add(Curve.Name);
+					CurveNamesToRemove.Add(Curve.GetName());
 				}
 			}
 		}
 
-		for (auto CurveName : CurveSmartNamesToRemove)
+		for (auto CurveName : CurveNamesToRemove)
 		{
 			const FAnimationCurveIdentifier CurveId(CurveName, ERawCurveTrackTypes::RCT_Float);
 			Controller.RemoveCurve(CurveId, bShouldTransact);
@@ -1852,15 +1823,15 @@ bool UnFbx::FFbxImporter::ImportAnimation(USkeleton* Skeleton, UAnimSequence * D
 	}
 	else
 	{
-	  // Store float curve tracks which use to exist on the animation
-	  for (const FFloatCurve& Curve : DestSeq->GetDataModel()->GetFloatCurves())
-	  {
-		  const FCurveMetaData* MetaData = MySkeleton->GetCurveMetaData(Curve.Name);
-		  if (MetaData && !MetaData->Type.bMorphtarget)
-		  {
-			  CurvesNotFound.Add(Curve.Name.DisplayName.ToString());
-		  }
-	  }
+		// Store float curve tracks which use to exist on the animation
+		for (const FFloatCurve& Curve : DestSeq->GetDataModel()->GetFloatCurves())
+		{
+			const FCurveMetaData* MetaData = MySkeleton->GetCurveMetaData(Curve.GetName());
+			if (MetaData && !MetaData->Type.bMorphtarget)
+			{
+				CurvesNotFound.Add(Curve.GetName().ToString());
+			}
+		}
 	}
 
 	if (bReimportWarnings && CurvesNotFound.Num())
@@ -1983,8 +1954,11 @@ void UnFbx::FFbxImporter::ImportBlendShapeCurves(FAnimCurveImportSettings& AnimI
 									if (ImportCurveToAnimSequence(AnimImportSettings.DestSeq, *ChannelName, Curve, 0, AnimImportSettings.AnimTimeSpan, bReimport, 0.01f /** for some reason blend shape values are coming as 100 scaled **/))
 									{
 										OutKeyCount = FMath::Max(OutKeyCount, Curve->KeyGetCount());
-										// this one doesn't reset Material curve to false, it just accumulate if true. 
-										MySkeleton->AccumulateCurveMetaData(*ChannelName, false, true);
+										if(ImportOptions->bAddCurveMetadataToSkeleton)
+										{
+											// this one doesn't reset Material curve to false, it just accumulate if true. 
+											MySkeleton->AccumulateCurveMetaData(*ChannelName, false, true);
+										}
 									}
 								}
 								else
@@ -2044,12 +2018,15 @@ void UnFbx::FFbxImporter::ImportBlendShapeCurves(FAnimCurveImportSettings& AnimI
 									if (ImportRichCurvesToAnimSequence(AnimImportSettings.DestSeq, CurveNames, Results, 0, bReimport))
 									{
 										OutKeyCount = FMath::Max(OutKeyCount, Curve->KeyGetCount());
-										for (const FString& CurveName : CurveNames)
+										if(ImportOptions->bAddCurveMetadataToSkeleton)
 										{
-											// this one doesn't reset Material curve to false, it just accumulate if true. 
-											MySkeleton->AccumulateCurveMetaData(*CurveName, false, true);
-										}	
-									}	
+											for (const FString& CurveName : CurveNames)
+											{
+												// this one doesn't reset Material curve to false, it just accumulate if true. 
+												MySkeleton->AccumulateCurveMetaData(*CurveName, false, true);
+											}
+										}
+									}
 								}
 							}
 						}
@@ -2086,11 +2063,11 @@ void UnFbx::FFbxImporter::ImportAnimationCustomAttribute(FAnimCurveImportSetting
 
 	for (const FFloatCurve& FloatCurve : CurveData.FloatCurves)
 	{
-		const FCurveMetaData* MetaData = MySkeleton->GetCurveMetaData(FloatCurve.Name);
+		const FCurveMetaData* MetaData = MySkeleton->GetCurveMetaData(FloatCurve.GetName());
 
 		if (MetaData && !MetaData->Type.bMorphtarget)
 		{
-			OutCurvesNotFound.Add(FloatCurve.Name.DisplayName.ToString());
+			OutCurvesNotFound.Add(FloatCurve.GetName().ToString());
 		}
 	}
 
@@ -2115,24 +2092,29 @@ void UnFbx::FFbxImporter::ImportAnimationCustomAttribute(FAnimCurveImportSetting
 				int32 CurveFlags = AACF_DefaultCurve;
 				if (ImportCurveToAnimSequence(DestSeq, InCurveName, InCurve, CurveFlags, AnimImportSettings.AnimTimeSpan, bReimport))
 				{
-					USkeleton* SeqSkeleton = DestSeq->GetSkeleton();
 					OutKeyCount = FMath::Max(OutKeyCount, InCurve->KeyGetCount());
-					// first let them override material curve if required
-					if (ImportOptions->bSetMaterialDriveParameterOnCustomAttribute)
-					{
-						// now mark this curve as material curve
-						SeqSkeleton->AccumulateCurveMetaData(FName(*InCurveName), true, false);
-					}
-					else
+
+					if(ImportOptions->bAddCurveMetadataToSkeleton)
 					{
-						// if not material set by default, apply naming convention for material
-						for (const auto& Suffix : ImportOptions->MaterialCurveSuffixes)
+						USkeleton* SeqSkeleton = DestSeq->GetSkeleton();
+
+						// first let them override material curve if required
+						if (ImportOptions->bSetMaterialDriveParameterOnCustomAttribute)
+						{
+							// now mark this curve as material curve
+							SeqSkeleton->AccumulateCurveMetaData(FName(*InCurveName), true, false);
+						}
+						else
 						{
-							int32 TotalSuffix = Suffix.Len();
-							if (InCurveName.Right(TotalSuffix) == Suffix)
+							// if not material set by default, apply naming convention for material
+							for (const auto& Suffix : ImportOptions->MaterialCurveSuffixes)
 							{
-								SeqSkeleton->AccumulateCurveMetaData(FName(*InCurveName), true, false);
-								break;
+								int32 TotalSuffix = Suffix.Len();
+								if (InCurveName.Right(TotalSuffix) == Suffix)
+								{
+									SeqSkeleton->AccumulateCurveMetaData(FName(*InCurveName), true, false);
+									break;
+								}
 							}
 						}
 					}
diff --git a/Engine/Source/Editor/UnrealEd/Public/FbxImporter.h b/Engine/Source/Editor/UnrealEd/Public/FbxImporter.h
--- a/Engine/Source/Editor/UnrealEd/Public/FbxImporter.h
+++ b/Engine/Source/Editor/UnrealEd/Public/FbxImporter.h
@@ -200,6 +200,7 @@ struct FBXImportOptions
 	bool	bDeleteExistingNonCurveCustomAttributes;
 	bool	bImportBoneTracks;
 	bool	bSetMaterialDriveParameterOnCustomAttribute;
+	bool	bAddCurveMetadataToSkeleton;
 	bool	bRemoveRedundantKeys;
 	bool	bDoNotImportCurveWithZero;
 	bool	bResetToFbxOnMaterialConflict;
diff --git a/Engine/Source/Runtime/AnimGraphRuntime/Private/AnimNodes/AnimNode_CopyPoseFromMesh.cpp b/Engine/Source/Runtime/AnimGraphRuntime/Private/AnimNodes/AnimNode_CopyPoseFromMesh.cpp
--- a/Engine/Source/Runtime/AnimGraphRuntime/Private/AnimNodes/AnimNode_CopyPoseFromMesh.cpp
+++ b/Engine/Source/Runtime/AnimGraphRuntime/Private/AnimNodes/AnimNode_CopyPoseFromMesh.cpp
@@ -1,6 +1,8 @@
 // Copyright Epic Games, Inc. All Rights Reserved.
 
 #include "AnimNodes/AnimNode_CopyPoseFromMesh.h"
+
+#include "Animation/AnimCurveUtils.h"
 #include "Animation/AnimInstanceProxy.h"
 #include "Animation/AnimTrace.h"
 #include "Components/SkeletalMeshComponent.h"
@@ -144,13 +146,14 @@ void FAnimNode_CopyPoseFromMesh::PreUpdate(const UAnimInstance* InAnimInstance)
 				UAnimInstance* SourceAnimInstance = CurrentMeshComponent->GetAnimInstance();
 				if (SourceAnimInstance)
 				{
-					// attribute curve contains all list
-					SourceCurveList.Reset();
-					SourceCurveList.Append(SourceAnimInstance->GetAnimationCurveList(EAnimCurveType::AttributeCurve));
+					// Potential optimization/tradeoff: If we stored the curve results on the mesh component in non-editor scenarios, this would be
+					// much faster (but take more memory). As it is, we need to translate the map stored on the anim instance.
+					const TMap<FName, float>& AnimCurveList = SourceAnimInstance->GetAnimationCurveList(EAnimCurveType::AttributeCurve);
+					UE::Anim::FCurveUtils::BuildUnsorted(SourceCurves, AnimCurveList);
 				}
 				else
 				{
-					SourceCurveList.Reset();
+					SourceCurves.Empty();
 				}
 			}
 
@@ -234,15 +237,7 @@ void FAnimNode_CopyPoseFromMesh::Evaluate_AnyThread(FPoseContext& Output)
 
 	if (bCopyCurves)
 	{
-		for (auto Iter = SourceCurveList.CreateConstIterator(); Iter; ++Iter)
-		{
-			const SmartName::UID_Type* UID = CurveNameToUIDMap.Find(Iter.Key());
-			if (UID)
-			{
-				// set source value to output curve
-				Output.Curve.Set(*UID, Iter.Value());
-			}
-		}
+		Output.Curve.CopyFrom(SourceCurves);
 	}
 
 	if (bCopyCustomAttributes)
@@ -268,7 +263,6 @@ void FAnimNode_CopyPoseFromMesh::ReinitializeMeshComponent(USkeletalMeshComponen
 	CurrentlyUsedSourceMesh.Reset();
 	CurrentlyUsedTargetMesh.Reset();
 	BoneMapToSource.Reset();
-	CurveNameToUIDMap.Reset();
 
 	if (TargetMeshComponent && IsValid(NewSourceMeshComponent) && NewSourceMeshComponent->GetSkeletalMeshAsset())
 	{
@@ -302,33 +296,6 @@ void FAnimNode_CopyPoseFromMesh::ReinitializeMeshComponent(USkeletalMeshComponen
 					}
 				}
 			}
-		
-			if (bCopyCurves)
-			{
-				USkeleton* SourceSkeleton = SourceSkelMesh->GetSkeleton();
-				USkeleton* TargetSkeleton = TargetSkelMesh->GetSkeleton();
-
-				// you shouldn't be here if this happened
-				if (ensureMsgf(SourceSkeleton, TEXT("Invalid null source skeleton : %s"), *GetNameSafe(SourceSkelMesh))
-					&& ensureMsgf(TargetSkeleton, TEXT("Invalid null target skeleton : %s"), *GetNameSafe(TargetSkelMesh)))
-				{
-					const FSmartNameMapping* SourceContainer = SourceSkeleton->GetSmartNameContainer(USkeleton::AnimCurveMappingName);
-					const FSmartNameMapping* TargetContainer = TargetSkeleton->GetSmartNameContainer(USkeleton::AnimCurveMappingName);
-
-					TArray<FName> SourceCurveNames;
-					SourceContainer->FillNameArray(SourceCurveNames);
-					for (int32 Index = 0; Index < SourceCurveNames.Num(); ++Index)
-					{
-						SmartName::UID_Type UID = TargetContainer->FindUID(SourceCurveNames[Index]);
-						if (UID != SmartName::MaxUID)
-						{
-							// has a valid UID, add to the list
-							SmartName::UID_Type& Value = CurveNameToUIDMap.Add(SourceCurveNames[Index]);
-							Value = UID;
-						}
-					}
-				}
-			}
 
 			if (bCopyCustomAttributes)
 			{
diff --git a/Engine/Source/Runtime/AnimGraphRuntime/Private/AnimNodes/AnimNode_CurveSource.cpp b/Engine/Source/Runtime/AnimGraphRuntime/Private/AnimNodes/AnimNode_CurveSource.cpp
--- a/Engine/Source/Runtime/AnimGraphRuntime/Private/AnimNodes/AnimNode_CurveSource.cpp
+++ b/Engine/Source/Runtime/AnimGraphRuntime/Private/AnimNodes/AnimNode_CurveSource.cpp
@@ -2,6 +2,7 @@
 
 #include "AnimNodes/AnimNode_CurveSource.h"
 #include "AnimationRuntime.h"
+#include "Animation/AnimCurveUtils.h"
 #include "Animation/AnimInstanceProxy.h"
 #include "Animation/AnimTrace.h"
 
@@ -84,21 +85,27 @@ void FAnimNode_CurveSource::Evaluate_AnyThread(FPoseContext& Output)
 		NamedCurveValues.Reset();
 		CurveSource->Execute_GetCurves(CurveSource.GetObject(), NamedCurveValues);
 
-		USkeleton* Skeleton = Output.AnimInstanceProxy->GetSkeleton();
+		const float ClampedAlpha = FMath::Clamp(Alpha, 0.0f, 1.0f);
 
-		for (FNamedCurveValue NamedValue : NamedCurveValues)
-		{
-			SmartName::UID_Type NameUID = Skeleton->GetUIDByName(USkeleton::AnimCurveMappingName, NamedValue.Name);
-			if (NameUID != SmartName::MaxUID)
+		FBlendedCurve Curve;
+		UE::Anim::FCurveUtils::BuildUnsorted(Curve, NamedCurveValues.Num(),
+			[&NamedCurveValues](int32 InCurveIndex)
 			{
-				const float CurrentValue = Output.Curve.Get(NameUID);
-				const float ClampedAlpha = FMath::Clamp(Alpha, 0.0f, 1.0f);
-				const float LerpedValue = FMath::Lerp(CurrentValue, NamedValue.Value, ClampedAlpha);
-				Output.Curve.Set(NameUID, LerpedValue);
+				return NamedCurveValues[InCurveIndex].Name;
+			},
+			[&NamedCurveValues](int32 InCurveIndex)
+			{
+				return NamedCurveValues[InCurveIndex].Value;
+			});
 
-				TRACE_ANIM_NODE_VALUE(Output, *NamedValue.Name.ToString(), LerpedValue);
-			}
+#if ANIM_TRACE_ENABLED
+		for (FNamedCurveValue NamedValue : NamedCurveValues)
+		{
+			TRACE_ANIM_NODE_VALUE(Output, *NamedValue.Name.ToString(), NamedValue.Value);
 		}
+#endif
+
+		Output.Curve.LerpTo(Curve, ClampedAlpha);
 	}
 }
 
diff --git a/Engine/Source/Runtime/AnimGraphRuntime/Private/AnimNodes/AnimNode_LayeredBoneBlend.cpp b/Engine/Source/Runtime/AnimGraphRuntime/Private/AnimNodes/AnimNode_LayeredBoneBlend.cpp
--- a/Engine/Source/Runtime/AnimGraphRuntime/Private/AnimNodes/AnimNode_LayeredBoneBlend.cpp
+++ b/Engine/Source/Runtime/AnimGraphRuntime/Private/AnimNodes/AnimNode_LayeredBoneBlend.cpp
@@ -4,10 +4,11 @@
 #include "AnimationRuntime.h"
 #include "Animation/AnimInstanceProxy.h"
 #include "Animation/AnimTrace.h"
+#include "Animation/AnimCurveTypes.h"
+#include "Animation/AnimCurveUtils.h"
 
 #include UE_INLINE_GENERATED_CPP_BY_NAME(AnimNode_LayeredBoneBlend)
 
-#define DEFAULT_SOURCEINDEX 0xFF
 /////////////////////////////////////////////////////
 // FAnimNode_LayeredBoneBlend
 
@@ -85,45 +86,26 @@ void FAnimNode_LayeredBoneBlend::UpdateCachedBoneData(const FBoneContainer& Requ
 	//Reinitialize bone blend weights now that we have cleared them
 	FAnimationRuntime::UpdateDesiredBoneWeight(DesiredBoneBlendWeights, CurrentBoneBlendWeights, BlendWeights);
 
-	TArray<uint16> const& CurveUIDFinder = RequiredBones.GetUIDToArrayLookupTable();
-	const int32 CurveUIDCount = CurveUIDFinder.Num();
-	const int32 TotalCount = FBlendedCurve::GetValidElementCount(&CurveUIDFinder);
-	if (TotalCount > 0)
+	// Build curve source indices
 	{
-		CurvePoseSourceIndices.Reset(TotalCount);
-		// initialize with FF - which is default
-		CurvePoseSourceIndices.Init(DEFAULT_SOURCEINDEX, TotalCount);
+		CurvePoseSourceIndices.Empty();
+		CurvePoseSourceIndices.Reserve(Skeleton->GetNumCurveMetaData());
 
-		// now go through point to correct source indices. Curve only picks one source index
-		for (int32 UIDIndex = 0; UIDIndex < CurveUIDCount; ++UIDIndex)
+		Skeleton->ForEachCurveMetaData([this, &RequiredBones](const FName& InCurveName, const FCurveMetaData& InMetaData)
 		{
-			int32 CurrentPoseIndex = CurveUIDFinder[UIDIndex];
-			if (CurrentPoseIndex != MAX_uint16)
+			for (const FBoneReference& LinkedBone : InMetaData.LinkedBones)
 			{
-				SmartName::UID_Type CurveUID = (SmartName::UID_Type)UIDIndex;
-
-				const FCurveMetaData* CurveMetaData = Skeleton->GetCurveMetaData(CurveUID);
-				if (CurveMetaData)
+				FCompactPoseBoneIndex CompactPoseIndex = LinkedBone.GetCompactPoseIndex(RequiredBones);
+				if (CompactPoseIndex != INDEX_NONE)
 				{
-					const TArray<FBoneReference>& LinkedBones = CurveMetaData->LinkedBones;
-					for (int32 LinkedBoneIndex = 0; LinkedBoneIndex < LinkedBones.Num(); ++LinkedBoneIndex)
+					if (DesiredBoneBlendWeights[CompactPoseIndex.GetInt()].BlendWeight > 0.f)
 					{
-						FCompactPoseBoneIndex CompactPoseIndex = LinkedBones[LinkedBoneIndex].GetCompactPoseIndex(RequiredBones);
-						if (CompactPoseIndex != INDEX_NONE)
-						{
-							if (DesiredBoneBlendWeights[CompactPoseIndex.GetInt()].BlendWeight > 0.f)
-							{
-								CurvePoseSourceIndices[CurrentPoseIndex] = IntCastChecked<uint8>(DesiredBoneBlendWeights[CompactPoseIndex.GetInt()].SourceIndex);
-							}
-						}
+						CurvePoseSourceIndices.Add(InCurveName, DesiredBoneBlendWeights[CompactPoseIndex.GetInt()].SourceIndex);
+						break;
 					}
 				}
 			}
-		}
-	}
-	else
-	{
-		CurvePoseSourceIndices.Reset();
+		});
 	}
 
 	RequiredBonesSerialNumber = RequiredBones.GetSerialNumber();
@@ -244,30 +226,23 @@ void FAnimNode_LayeredBoneBlend::Evaluate_AnyThread(FPoseContext& Output)
 			}
 		}
 
-		// filter to make sure it only includes curves that is linked to the correct bone filter
-		TArray<uint16> const* CurveUIDFinder = Output.Curve.UIDToArrayIndexLUT;
-		const int32 TotalCount = Output.Curve.NumValidCurveCount;
-		// now go through point to correct source indices. Curve only picks one source index
-		for (USkeleton::AnimCurveUID UIDIndex = 0; UIDIndex < CurveUIDFinder->Num(); ++UIDIndex)
-		{
-			int32 CurvePoseIndex = Output.Curve.GetArrayIndexByUID(UIDIndex);
-			if (CurvePoseSourceIndices.IsValidIndex(CurvePoseIndex))
+		// filter to make sure it only includes curves that are linked to the correct bone filter
+		UE::Anim::FNamedValueArrayUtils::RemoveByPredicate(BasePoseContext.Curve, CurvePoseSourceIndices,
+			[](const UE::Anim::FCurveElement& InOutBasePoseElement, const UE::Anim::FCurveElementIndexed& InSourceIndexElement)
 			{
-				int32 SourceIndex = CurvePoseSourceIndices[CurvePoseIndex];
-				if (SourceIndex != DEFAULT_SOURCEINDEX)
+				// if source index is set, remove base pose curve value
+				return (InSourceIndexElement.Index != INDEX_NONE);
+			});
+
+		// Filter child pose curves
+		for (int32 ChildIndex = 0; ChildIndex < NumPoses; ++ChildIndex)
+		{
+			UE::Anim::FNamedValueArrayUtils::RemoveByPredicate(TargetBlendCurves[ChildIndex], CurvePoseSourceIndices,
+				[ChildIndex](const UE::Anim::FCurveElement& InOutBasePoseElement, const UE::Anim::FCurveElementIndexed& InSourceIndexElement)
 				{
-					// if source index is set, invalidate base pose curve value
-					BasePoseContext.Curve.InvalidateCurveWeight(UIDIndex);
-					for (int32 ChildIndex = 0; ChildIndex < NumPoses; ++ChildIndex)
-					{
-						if (SourceIndex != ChildIndex)
-						{
-							// if not source, invalidate it
-							TargetBlendCurves[ChildIndex].InvalidateCurveWeight(UIDIndex);
-						}
-					}
-				}
-			}
+					// if not source, remove it
+					return (InSourceIndexElement.Index != INDEX_NONE) && (InSourceIndexElement.Index != ChildIndex);
+				});
 		}
 
 		FAnimationRuntime::EBlendPosesPerBoneFilterFlags BlendFlags = FAnimationRuntime::EBlendPosesPerBoneFilterFlags::None;
diff --git a/Engine/Source/Runtime/AnimGraphRuntime/Private/AnimNodes/AnimNode_ModifyCurve.cpp b/Engine/Source/Runtime/AnimGraphRuntime/Private/AnimNodes/AnimNode_ModifyCurve.cpp
--- a/Engine/Source/Runtime/AnimGraphRuntime/Private/AnimNodes/AnimNode_ModifyCurve.cpp
+++ b/Engine/Source/Runtime/AnimGraphRuntime/Private/AnimNodes/AnimNode_ModifyCurve.cpp
@@ -3,6 +3,7 @@
 #include "AnimNodes/AnimNode_ModifyCurve.h"
 #include "AnimationRuntime.h"
 #include "Animation/AnimInstanceProxy.h"
+#include "Animation/AnimCurveTypes.h"
 
 #include UE_INLINE_GENERATED_CPP_BY_NAME(AnimNode_ModifyCurve)
 
@@ -10,7 +11,6 @@ FAnimNode_ModifyCurve::FAnimNode_ModifyCurve()
 {
 	ApplyMode = EModifyCurveApplyMode::Blend;
 	Alpha = 1.f;
-	bInitializeLastValuesMap = false;
 }
 
 void FAnimNode_ModifyCurve::Initialize_AnyThread(const FAnimationInitializeContext& Context)
@@ -22,11 +22,7 @@ void FAnimNode_ModifyCurve::Initialize_AnyThread(const FAnimationInitializeConte
 	// Init our last values array to be the right size
 	if (ApplyMode == EModifyCurveApplyMode::WeightedMovingAverage)
 	{
-		LastCurveValues.Reset(CurveValues.Num());
-		LastCurveValues.AddZeroed(CurveValues.Num());
-
-		// Indicate that last values for curve map should be initialized on evaluate since updated CurveMap is not yet available.
-		bInitializeLastValuesMap = true;		
+		LastCurve.Empty();
 	}
 }
 
@@ -45,106 +41,85 @@ void FAnimNode_ModifyCurve::Evaluate_AnyThread(FPoseContext& Output)
 
 	Output = SourceData;
 
-	//	Morph target and Material parameter curves
-	USkeleton* Skeleton = Output.AnimInstanceProxy->GetSkeleton();
+	// Build internal curves from array & map
+	FBlendedCurve Curve;
+	Curve.Reserve(CurveNames.Num());
+	for(int32 CurveIndex = 0; CurveIndex < CurveNames.Num(); ++CurveIndex)
+	{
+		Curve.Add(CurveNames[CurveIndex], CurveValues[CurveIndex]);
+	}
 
-	// Check if Last Values Map should be initialized. This has to be done if curve map is changed in the node graph without compiling.
-	if (bInitializeLastValuesMap)
+	if(CurveMap.Num() > 0)
 	{
-		// CurveMap is initialized so Last curve map values can be stored.
-		LastCurveMapValues.Reset();
-		LastCurveMapValues.Reserve(CurveMap.Num());
-		for (auto It = CurveMap.CreateConstIterator(); It; ++It)
+		FBlendedCurve MapCurve;
+		MapCurve.Reserve(CurveNames.Num());
+		for(const TPair<FName, float>& NameValuePair : CurveMap)
 		{
-			LastCurveMapValues.Add(It.Key(), 0.0f);
+			MapCurve.Add(NameValuePair.Key, NameValuePair.Value);
 		}
-		bInitializeLastValuesMap = false;
+
+		// Combine pin & map curves in case they overlap
+		Curve.Combine(MapCurve);
 	}
 
-	// Process curve map if available.
-	for (auto& CurveNameValue : CurveMap)
-	{				
-		SmartName::UID_Type NameUID = Skeleton->GetUIDByName(USkeleton::AnimCurveMappingName, CurveNameValue.Key);
-		if (NameUID != SmartName::MaxUID)
-		{
-			float CurrentValue = Output.Curve.Get(NameUID);
-			float NewValue = CurveNameValue.Value;
-			
-			if (ApplyMode == EModifyCurveApplyMode::WeightedMovingAverage)
+	if (ApplyMode == EModifyCurveApplyMode::WeightedMovingAverage)
+	{
+		UE::Anim::FNamedValueArrayUtils::Union(Output.Curve, Curve, LastCurve,
+			[this](UE::Anim::FCurveElement& InOutResult, const UE::Anim::FCurveElement& InCurveElement, const UE::Anim::FCurveElement& InLastCurveElement, UE::Anim::ENamedValueUnionFlags InFlags)
 			{
-				float* LastValue = LastCurveMapValues.Find(CurveNameValue.Key);
-				if (LastValue)
+				// Only apply curves that we are overriding
+				if(EnumHasAnyFlags(InFlags, UE::Anim::ENamedValueUnionFlags::ValidArg1))
 				{
-					ProcessCurveWMAOperation(Output, NameUID, CurrentValue, NewValue, *LastValue);	
+					InOutResult.Value = ProcessCurveWMAOperation(InCurveElement.Value, InLastCurveElement.Value);
 				}
-			}
-			else
-			{
-				ProcessCurveOperation(ApplyMode, Output, NameUID, CurrentValue, NewValue);
-			}
-		}
+			});
 	}
-
-	// Process custom selected curve pins.
-	check(CurveNames.Num() == CurveValues.Num());
-
-	for (int32 ModIdx = 0; ModIdx < CurveNames.Num(); ModIdx++)
+	else
 	{
-		FName CurveName = CurveNames[ModIdx];
-		SmartName::UID_Type NameUID = Skeleton->GetUIDByName(USkeleton::AnimCurveMappingName, CurveName);
-		if (NameUID != SmartName::MaxUID)
-		{
-			float NewValue = CurveValues[ModIdx];
-			float CurrentValue = Output.Curve.Get(NameUID);
-			if (ApplyMode == EModifyCurveApplyMode::WeightedMovingAverage)
-			{
-				check(LastCurveValues.Num() == CurveValues.Num())
-				float& LastValue = LastCurveValues[ModIdx];
-				ProcessCurveWMAOperation(Output, NameUID, CurrentValue, NewValue, LastValue);
-			}
-			else
+		UE::Anim::FNamedValueArrayUtils::Union(Output.Curve, Curve,
+			[this](UE::Anim::FCurveElement& InOutResult, const UE::Anim::FCurveElement& InCurveElement, UE::Anim::ENamedValueUnionFlags InFlags)
 			{
-				ProcessCurveOperation(ApplyMode, Output, NameUID, CurrentValue, NewValue);
-			}
-		}
+				// Only apply curves that we are overriding
+				if(EnumHasAnyFlags(InFlags, UE::Anim::ENamedValueUnionFlags::ValidArg1))
+				{
+					InOutResult.Value = ProcessCurveOperation(InOutResult.Value, InCurveElement.Value);
+				}
+			});
 	}
+
+	LastCurve.CopyFrom(Curve);
 }
 
-void FAnimNode_ModifyCurve::ProcessCurveOperation(const EModifyCurveApplyMode& InApplyMode, FPoseContext& Output, const SmartName::UID_Type& NameUID, float CurrentValue, float NewValue)
+float FAnimNode_ModifyCurve::ProcessCurveOperation(float CurrentValue, float NewValue) const
 {
 	float UseNewValue = CurrentValue;
 
 	// Use ApplyMode enum to decide how to apply
-	if (InApplyMode == EModifyCurveApplyMode::Add)
+	if (ApplyMode == EModifyCurveApplyMode::Add)
 	{
 		UseNewValue = CurrentValue + NewValue;
 	}
-	else if (InApplyMode == EModifyCurveApplyMode::Scale)
+	else if (ApplyMode == EModifyCurveApplyMode::Scale)
 	{
 		UseNewValue = CurrentValue * NewValue;
 	}
-	else if (InApplyMode == EModifyCurveApplyMode::RemapCurve)
+	else if (ApplyMode == EModifyCurveApplyMode::RemapCurve)
 	{
 		const float RemapScale = 1.f / FMath::Max(1.f - NewValue, 0.01f);
 		UseNewValue = FMath::Min(FMath::Max(CurrentValue - NewValue, 0.f) * RemapScale, 1.f);
 	}
-	else if (InApplyMode == EModifyCurveApplyMode::Blend)
+	else if (ApplyMode == EModifyCurveApplyMode::Blend)
 	{
 		UseNewValue = NewValue;
 	}
 
-	float UseAlpha = FMath::Clamp(Alpha, 0.f, 1.f);
-	Output.Curve.Set(NameUID, FMath::Lerp(CurrentValue, UseNewValue, UseAlpha));
+	const float UseAlpha = FMath::Clamp(Alpha, 0.f, 1.f);
+	return FMath::Lerp(CurrentValue, UseNewValue, UseAlpha);
 }
 
-
-void FAnimNode_ModifyCurve::ProcessCurveWMAOperation(FPoseContext& Output, const SmartName::UID_Type& NameUID, float CurrentValue, float NewValue, float& InOutLastValue)
+float FAnimNode_ModifyCurve::ProcessCurveWMAOperation(float CurrentValue, float LastValue) const
 {
-	const float WAvg = FMath::WeightedMovingAverage(CurrentValue, InOutLastValue, Alpha);
-	// Update the last curve value for next run
-	InOutLastValue = WAvg;
-
-	Output.Curve.Set(NameUID, WAvg);
+	return FMath::WeightedMovingAverage(CurrentValue, LastValue, Alpha);
 }
 
 void FAnimNode_ModifyCurve::Update_AnyThread(const FAnimationUpdateContext& Context)
diff --git a/Engine/Source/Runtime/AnimGraphRuntime/Private/AnimNodes/AnimNode_PoseBlendNode.cpp b/Engine/Source/Runtime/AnimGraphRuntime/Private/AnimNodes/AnimNode_PoseBlendNode.cpp
--- a/Engine/Source/Runtime/AnimGraphRuntime/Private/AnimNodes/AnimNode_PoseBlendNode.cpp
+++ b/Engine/Source/Runtime/AnimGraphRuntime/Private/AnimNodes/AnimNode_PoseBlendNode.cpp
@@ -2,6 +2,7 @@
 
 #include "AnimNodes/AnimNode_PoseBlendNode.h"
 #include "AnimationRuntime.h"
+#include "Animation/AnimCurveUtils.h"
 #include "Animation/AnimInstanceProxy.h"
 
 #include UE_INLINE_GENERATED_CPP_BY_NAME(AnimNode_PoseBlendNode)
@@ -20,6 +21,13 @@ void FAnimNode_PoseBlendNode::Initialize_AnyThread(const FAnimationInitializeCon
 	DECLARE_SCOPE_HIERARCHICAL_COUNTER_ANIMNODE(Initialize_AnyThread)
 	FAnimNode_PoseHandler::Initialize_AnyThread(Context);
 
+	BulkCurves.Empty();
+	for (int32 PoseIdx = 0; PoseIdx < PoseExtractContext.PoseCurves.Num(); ++PoseIdx)
+	{
+		FPoseCurve& PoseCurve = PoseExtractContext.PoseCurves[PoseIdx];
+		BulkCurves.Add(PoseCurve.Name, BulkCurves.Num());
+	}
+
 	SourcePose.Initialize(Context);
 }
 
@@ -50,18 +58,14 @@ void FAnimNode_PoseBlendNode::Evaluate_AnyThread(FPoseContext& Output)
 	if (CachedPoseAsset && PoseExtractContext.PoseCurves.Num() > 0 && CachedPoseAsset->GetSkeleton() != nullptr)
 	{
 		FPoseContext CurrentPose(Output);
-		// only give pose curve, we don't set any more curve here
-		for (int32 PoseIdx = 0; PoseIdx < PoseExtractContext.PoseCurves.Num(); ++PoseIdx)
+
+		// Remap incoming curve
+		UE::Anim::FCurveUtils::BulkGet(SourceData.Curve, BulkCurves, [this](const UE::Anim::FNamedIndexElement& InBulkElement, float InValue)
 		{
-			FPoseCurve& PoseCurve = PoseExtractContext.PoseCurves[PoseIdx];
-			// Get value of input curve
-			float InputValue = SourceData.Curve.Get(PoseCurve.UID);
 			// Remap using chosen BlendOption
-			float RemappedValue = FAlphaBlend::AlphaToBlendOption(InputValue, BlendOption, CustomCurve);
-
-			PoseCurve.Value = RemappedValue;
-		}
-
+			const float RemappedValue = FAlphaBlend::AlphaToBlendOption(InValue, BlendOption, CustomCurve);
+			PoseExtractContext.PoseCurves[InBulkElement.Index].Value = RemappedValue;
+		});
 
 		FAnimationPoseData CurrentAnimationPoseData(CurrentPose);
 		if (CachedPoseAsset->GetAnimationPose(CurrentAnimationPoseData, PoseExtractContext))
diff --git a/Engine/Source/Runtime/AnimGraphRuntime/Private/AnimNodes/AnimNode_PoseByName.cpp b/Engine/Source/Runtime/AnimGraphRuntime/Private/AnimNodes/AnimNode_PoseByName.cpp
--- a/Engine/Source/Runtime/AnimGraphRuntime/Private/AnimNodes/AnimNode_PoseByName.cpp
+++ b/Engine/Source/Runtime/AnimGraphRuntime/Private/AnimNodes/AnimNode_PoseByName.cpp
@@ -18,13 +18,12 @@ void FAnimNode_PoseByName::Initialize_AnyThread(const FAnimationInitializeContex
 void FAnimNode_PoseByName::RebuildPoseList(const FBoneContainer& InBoneContainer, const UPoseAsset* InPoseAsset)
 {
 	PoseExtractContext.PoseCurves.Reset();
-	const TArray<FSmartName>& PoseNames = InPoseAsset->GetPoseNames();
+	const TArray<FName>& PoseNames = InPoseAsset->GetPoseFNames();
 	const int32 PoseIndex = InPoseAsset->GetPoseIndexByName(PoseName);
-	TArray<uint16> const& LUTIndex = InBoneContainer.GetUIDToArrayLookupTable();
-	if (PoseIndex != INDEX_NONE && LUTIndex.IsValidIndex(PoseNames[PoseIndex].UID) && LUTIndex[PoseNames[PoseIndex].UID] != MAX_uint16)
+	if (PoseIndex != INDEX_NONE)
 	{
 		// we keep pose index as that is the fastest way to search when extracting pose asset
-		PoseExtractContext.PoseCurves.Add(FPoseCurve(PoseIndex, PoseNames[PoseIndex].UID, 0.f));
+		PoseExtractContext.PoseCurves.Add(FPoseCurve(PoseIndex, PoseNames[PoseIndex], 0.f));
 	}
 }
 
