Fix for crash when showing Weighted Bones in Skeleletal Mesh Editor

diff --git a/Engine/Source/Editor/SkeletonEditor/Private/SSkeletonTree.cpp b/Engine/Source/Editor/SkeletonEditor/Private/SSkeletonTree.cpp
--- a/Engine/Source/Editor/SkeletonEditor/Private/SSkeletonTree.cpp
+++ b/Engine/Source/Editor/SkeletonEditor/Private/SSkeletonTree.cpp
@@ -2421,7 +2421,7 @@ ESkeletonTreeFilterResult SSkeletonTree::HandleFilterSkeletonTreeItem(const FSke
 				UDebugSkelMeshComponent* PreviewMeshComponent = GetPreviewScene()->GetPreviewMeshComponent();
 				if (PreviewMeshComponent)
 				{
-					int32 BoneMeshIndex = PreviewMeshComponent->GetBoneIndex(BoneItem->GetRowItemName());
+					const int32 BoneMeshIndex = PreviewMeshComponent->GetBoneIndex(BoneItem->GetRowItemName());
 
 					// Remove non-mesh bones if we're filtering
 					if ((BoneFilter == EBoneFilter::Mesh || BoneFilter == EBoneFilter::Weighted || BoneFilter == EBoneFilter::LOD) &&
@@ -2431,13 +2431,13 @@ ESkeletonTreeFilterResult SSkeletonTree::HandleFilterSkeletonTreeItem(const FSke
 					}
 
 					// Remove non-vertex-weighted bones if we're filtering
-					if (BoneFilter == EBoneFilter::Weighted && !BoneItem->IsBoneWeighted(BoneMeshIndex, PreviewMeshComponent))
+					else if (BoneFilter == EBoneFilter::Weighted && !BoneItem->IsBoneWeighted(BoneMeshIndex, PreviewMeshComponent))
 					{
 						Result = ESkeletonTreeFilterResult::Hidden;
 					}
 
 					// Remove non-vertex-weighted bones if we're filtering
-					if (BoneFilter == EBoneFilter::LOD && !BoneItem->IsBoneRequired(BoneMeshIndex, PreviewMeshComponent))
+					else if (BoneFilter == EBoneFilter::LOD && !BoneItem->IsBoneRequired(BoneMeshIndex, PreviewMeshComponent))
 					{
 						Result = ESkeletonTreeFilterResult::Hidden;
 					}
@@ -2454,18 +2454,18 @@ ESkeletonTreeFilterResult SSkeletonTree::HandleFilterSkeletonTreeItem(const FSke
 			}
 
 			// Remove non-mesh sockets if we're filtering
-			if ((SocketFilter == ESocketFilter::Mesh || SocketFilter == ESocketFilter::None) && SocketItem->GetParentType() == ESocketParentType::Skeleton)
+			else if ((SocketFilter == ESocketFilter::Mesh || SocketFilter == ESocketFilter::None) && SocketItem->GetParentType() == ESocketParentType::Skeleton)
 			{
 				Result = ESkeletonTreeFilterResult::Hidden;
 			}
 
 			// Remove non-skeleton sockets if we're filtering
-			if ((SocketFilter == ESocketFilter::Skeleton || SocketFilter == ESocketFilter::None) && SocketItem->GetParentType() == ESocketParentType::Mesh)
+			else if ((SocketFilter == ESocketFilter::Skeleton || SocketFilter == ESocketFilter::None) && SocketItem->GetParentType() == ESocketParentType::Mesh)
 			{
 				Result = ESkeletonTreeFilterResult::Hidden;
 			}
 
-			if (SocketFilter == ESocketFilter::Active && SocketItem->GetParentType() == ESocketParentType::Skeleton && SocketItem->IsSocketCustomized())
+			else if (SocketFilter == ESocketFilter::Active && SocketItem->GetParentType() == ESocketParentType::Skeleton && SocketItem->IsSocketCustomized())
 			{
 				// Don't add the skeleton socket if it's already added for the mesh
 				Result = ESkeletonTreeFilterResult::Hidden;
diff --git a/Engine/Source/Editor/SkeletonEditor/Private/SkeletonTreeBoneItem.cpp b/Engine/Source/Editor/SkeletonEditor/Private/SkeletonTreeBoneItem.cpp
--- a/Engine/Source/Editor/SkeletonEditor/Private/SkeletonTreeBoneItem.cpp
+++ b/Engine/Source/Editor/SkeletonEditor/Private/SkeletonTreeBoneItem.cpp
@@ -502,6 +502,11 @@ FReply FSkeletonTreeBoneItem::HandleDrop(const FDragDropEvent& DragDropEvent)
 bool FSkeletonTreeBoneItem::IsBoneWeighted(int32 MeshBoneIndex, UDebugSkelMeshComponent* PreviewComponent)
 {
 	// MeshBoneIndex must be an index into the mesh's skeleton, *not* the source skeleton!!!
+	if (MeshBoneIndex == INDEX_NONE)
+	{
+		// If we get an invalid index, we are done here
+		return false; 
+	}
 
 	if (!PreviewComponent || !PreviewComponent->GetSkeletalMeshAsset() || !PreviewComponent->GetSkeletalMeshAsset()->GetResourceForRendering() || !PreviewComponent->GetSkeletalMeshAsset()->GetResourceForRendering()->LODRenderData.Num())
 	{
