Fix for add virtual bone menu sometimes not showing

diff --git a/Engine/Source/Editor/SkeletonEditor/Private/SSkeletonTree.cpp b/Engine/Source/Editor/SkeletonEditor/Private/SSkeletonTree.cpp
--- a/Engine/Source/Editor/SkeletonEditor/Private/SSkeletonTree.cpp
+++ b/Engine/Source/Editor/SkeletonEditor/Private/SSkeletonTree.cpp
@@ -762,7 +762,11 @@ TSharedPtr< SWidget > SSkeletonTree::CreateContextMenu()
 			{
 				MenuBuilder.AddSubMenu(LOCTEXT("AddVirtualBone", "Add Virtual Bone"),
 					LOCTEXT("AddVirtualBone_ToolTip", "Adds a virtual bone to the skeleton."),
-					FNewMenuDelegate::CreateSP(this, &SSkeletonTree::FillVirtualBoneSubmenu));
+					FNewMenuDelegate::CreateLambda([this](FMenuBuilder& MenuBuilder)
+						{
+							TSharedRef<SBoneTreeMenu> MenuContent = SSkeletonTree::CreateVirtualBoneMenu(this);
+							MenuBuilder.AddWidget(MenuContent, FText::GetEmpty(), true);
+						}));
 			}
 
 			MenuBuilder.EndSection();
@@ -963,25 +967,24 @@ bool GetSourceNameFromItem(TSharedPtr<ISkeletonTreeItem> SourceBone, FName& OutN
 	return false;
 }
 
-void SSkeletonTree::FillVirtualBoneSubmenu(FMenuBuilder& MenuBuilder)
+TSharedRef<SBoneTreeMenu> SSkeletonTree::CreateVirtualBoneMenu(SSkeletonTree* InSkeletonTree)
 {
-	TArray<TSharedPtr<ISkeletonTreeItem>> SelectedItems = SkeletonTreeView->GetSelectedItems();
+	const TArray<TSharedPtr<ISkeletonTreeItem>> SelectedItems = InSkeletonTree->GetSelectedItems();
 
-	const bool bShowVirtualBones = false;
 	TSharedRef<SBoneTreeMenu> MenuContent = SNew(SBoneTreeMenu)
-	.bShowVirtualBones(false)
-	.Title(LOCTEXT("TargetBonePickerTitle", "Pick Target Bone..."))
-	.OnBoneSelectionChanged(this, &SSkeletonTree::OnVirtualTargetBonePicked, SelectedItems)
-	.OnGetReferenceSkeleton(this, &SSkeletonTree::OnGetReferenceSkeleton);
-	MenuBuilder.AddWidget(MenuContent, FText::GetEmpty(), true);
+		.bShowVirtualBones(false)
+		.Title(LOCTEXT("TargetBonePickerTitle", "Pick Target Bone..."))
+		.OnBoneSelectionChanged(InSkeletonTree, &SSkeletonTree::OnVirtualTargetBonePicked, SelectedItems)
+		.OnGetReferenceSkeleton(InSkeletonTree, &SSkeletonTree::OnGetReferenceSkeleton);
 
 	MenuContent->RegisterActiveTimer(0.0f, FWidgetActiveTimerDelegate::CreateLambda(
 		[FilterTextBox = MenuContent->GetFilterTextWidget()](double, float)
 		{
 			FSlateApplication::Get().SetKeyboardFocus(FilterTextBox);
 			return EActiveTimerReturnType::Stop;
-		}
-		));
+		}));
+
+	return MenuContent;
 }
 
 void SSkeletonTree::OnVirtualTargetBonePicked(FName TargetBoneName, TArray<TSharedPtr<ISkeletonTreeItem>> SourceBones)
@@ -1864,7 +1867,24 @@ void SSkeletonTree::RegisterNewMenu()
 			"VirtualBones",
 			LOCTEXT("AddVirtualBone", "Add Virtual Bone"),
 			LOCTEXT("AddVirtualBone_ToolTip", "Adds a virtual bone to the skeleton."),
-			FNewToolMenuChoice(FNewMenuDelegate::CreateSP(this, &SSkeletonTree::FillVirtualBoneSubmenu)));
+			FNewToolMenuDelegate::CreateLambda([](UToolMenu* InMenu)
+				{
+					USkeletonTreeMenuContext* MenuContext = InMenu->Context.FindContext<USkeletonTreeMenuContext>();
+                    if(MenuContext == nullptr)
+                    {
+                    	return;
+                    }
+                
+                    TSharedPtr<SSkeletonTree> SkeletonTree = MenuContext->SkeletonTree.Pin();
+                    if(!SkeletonTree.IsValid())
+                    {
+                    	return;
+                    }
+
+					TSharedRef<SBoneTreeMenu> MenuContent = SSkeletonTree::CreateVirtualBoneMenu(SkeletonTree.Get());
+					FToolMenuEntry WidgetEntry = FToolMenuEntry::InitWidget("VirtualBones", MenuContent, FText());
+					InMenu->AddMenuEntry(NAME_None, WidgetEntry);
+				}));
 	}
 
 	{
diff --git a/Engine/Source/Editor/SkeletonEditor/Private/SSkeletonTree.h b/Engine/Source/Editor/SkeletonEditor/Private/SSkeletonTree.h
--- a/Engine/Source/Editor/SkeletonEditor/Private/SSkeletonTree.h
+++ b/Engine/Source/Editor/SkeletonEditor/Private/SSkeletonTree.h
@@ -217,7 +217,7 @@ private:
 	void OnPromoteSocket();
 
 	/** Create sub menu to allow users to pick a target bone for the new space switching bone(s) */
-	void FillVirtualBoneSubmenu(FMenuBuilder& MenuBuilder);
+	static TSharedRef<class SBoneTreeMenu> CreateVirtualBoneMenu(SSkeletonTree* InSkeletonTree);
 
 	/** Handler for user picking a target bone */
 	void OnVirtualTargetBonePicked(FName TargetBoneName, TArray<TSharedPtr<class ISkeletonTreeItem>> SourceBones);
