Fix crash deleting sync marker at SyncMarkers Window, by changing a CastChecked by a Cast, so it does not crashes with AnimMontages (as GetCompatibleAnimSequences retuirns all UAnimSequenceBase assets)

diff --git a/Engine/Source/Editor/SkeletonEditor/Private/EditableSkeleton.cpp b/Engine/Source/Editor/SkeletonEditor/Private/EditableSkeleton.cpp
--- a/Engine/Source/Editor/SkeletonEditor/Private/EditableSkeleton.cpp
+++ b/Engine/Source/Editor/SkeletonEditor/Private/EditableSkeleton.cpp
@@ -946,7 +946,7 @@ int32 FEditableSkeleton::DeleteSyncMarkers(const TArray<FName>& InSyncMarkerName
 			const FAssetData& PossibleAnimSequence = CompatibleAnimSequences[AssetIndex];
 			if(UObject* LoadedAsset = PossibleAnimSequence.GetAsset())
 			{
-				if(UAnimSequence* Sequence = CastChecked<UAnimSequence>(LoadedAsset))
+				if(UAnimSequence* Sequence = Cast<UAnimSequence>(LoadedAsset))
 				{
 					if (Sequence->RemoveSyncMarkers(InSyncMarkerNames))
 					{
