Stop editing of virtual bone transforms via details panel customization

diff --git a/Engine/Source/Editor/SkeletonEditor/Private/BoneProxy.cpp b/Engine/Source/Editor/SkeletonEditor/Private/BoneProxy.cpp
--- a/Engine/Source/Editor/SkeletonEditor/Private/BoneProxy.cpp
+++ b/Engine/Source/Editor/SkeletonEditor/Private/BoneProxy.cpp
@@ -20,6 +20,7 @@ UBoneProxy::UBoneProxy()
 	, PreviousScale(FVector::ZeroVector)
 	, bManipulating(false)
 	, bIsTickable(false)
+	, bIsTransformEditable(true)
 {
 }
 
@@ -240,7 +241,7 @@ void UBoneProxy::OnNumericValueCommitted(
 	ETransformType TransformType,
 	bool bIsCommit)
 {
-	if(TransformType != TransformType_Bone)
+	if(TransformType != TransformType_Bone || !bIsTransformEditable)
 	{
 		return;
 	}
@@ -346,7 +347,7 @@ void UBoneProxy::OnMultiNumericValueCommitted(
 
 bool UBoneProxy::DiffersFromDefault(ESlateTransformComponent::Type Component, ETransformType TransformType) const
 {
-	if(TransformType == TransformType_Bone)
+	if(TransformType == TransformType_Bone && bIsTransformEditable)
 	{
 		switch(Component)
 		{
@@ -375,7 +376,7 @@ bool UBoneProxy::DiffersFromDefault(ESlateTransformComponent::Type Component, ET
 
 void UBoneProxy::ResetToDefault(ESlateTransformComponent::Type InComponent, ETransformType TransformType)
 {
-	if(TransformType == TransformType_Bone)
+	if(TransformType == TransformType_Bone && bIsTransformEditable)
 	{
 		if (UDebugSkelMeshComponent* Component = SkelMeshComponent.Get())
 		{
diff --git a/Engine/Source/Editor/SkeletonEditor/Private/BoneProxyDetailsCustomization.cpp b/Engine/Source/Editor/SkeletonEditor/Private/BoneProxyDetailsCustomization.cpp
--- a/Engine/Source/Editor/SkeletonEditor/Private/BoneProxyDetailsCustomization.cpp
+++ b/Engine/Source/Editor/SkeletonEditor/Private/BoneProxyDetailsCustomization.cpp
@@ -181,6 +181,20 @@ void FBoneProxyDetailsCustomization::CustomizeDetails(IDetailLayoutBuilder& Deta
 	{
 		bIsEditingEnabled = (Component->AnimScriptInstance == Component->PreviewInstance);
 	}
+
+	if (bIsEditingEnabled)
+	{
+		bool bHasEditable = false;
+		for (const UBoneProxy* BoneProxy : BoneProxies)
+		{
+			if (BoneProxy->bIsTransformEditable)
+			{
+				bHasEditable = true;
+				break;
+			}
+		}
+		bIsEditingEnabled = bHasEditable;
+	}
 	
 	DetailBuilder.HideCategory(TEXT("Transform"));
 	DetailBuilder.HideCategory(TEXT("Reference Transform"));
@@ -377,6 +391,10 @@ void FBoneProxyDetailsCustomization::CustomizeDetails(IDetailLayoutBuilder& Deta
 
 			for (UBoneProxy* BoneProxy : BoneProxiesView)
 			{
+				if (!BoneProxy->bIsTransformEditable)
+				{
+					continue; // Skip non editable bones
+				}
 				if (const UDebugSkelMeshComponent* Component = BoneProxy->SkelMeshComponent.Get())
 				{
 					switch(InComponent)
@@ -484,7 +502,7 @@ bool FBoneProxyDetailsCustomization::IsResetLocationVisible(TSharedPtr<IProperty
 		{
 			if (FAnimNode_ModifyBone* ModifyBone = Component->PreviewInstance->FindModifiedBone(BoneProxy->BoneName))
 			{
-				if (ModifyBone->Translation != FVector::ZeroVector)
+				if (ModifyBone->Translation != FVector::ZeroVector && BoneProxy->bIsTransformEditable)
 				{
 					return true;
 				}
@@ -503,7 +521,7 @@ bool FBoneProxyDetailsCustomization::IsResetRotationVisible(TSharedPtr<IProperty
 		{
 			if (FAnimNode_ModifyBone* ModifyBone = Component->PreviewInstance->FindModifiedBone(BoneProxy->BoneName))
 			{
-				if (ModifyBone->Rotation != FRotator::ZeroRotator)
+				if (ModifyBone->Rotation != FRotator::ZeroRotator && BoneProxy->bIsTransformEditable)
 				{
 					return true;
 				}
@@ -522,7 +540,7 @@ bool FBoneProxyDetailsCustomization::IsResetScaleVisible(TSharedPtr<IPropertyHan
 		{
 			if (FAnimNode_ModifyBone* ModifyBone = Component->PreviewInstance->FindModifiedBone(BoneProxy->BoneName))
 			{
-				if (ModifyBone->Scale != FVector(1.0f))
+				if (ModifyBone->Scale != FVector(1.0f) && BoneProxy->bIsTransformEditable)
 				{
 					return true;
 				}
diff --git a/Engine/Source/Editor/SkeletonEditor/Private/SkeletonTreeVirtualBoneItem.cpp b/Engine/Source/Editor/SkeletonEditor/Private/SkeletonTreeVirtualBoneItem.cpp
--- a/Engine/Source/Editor/SkeletonEditor/Private/SkeletonTreeVirtualBoneItem.cpp
+++ b/Engine/Source/Editor/SkeletonEditor/Private/SkeletonTreeVirtualBoneItem.cpp
@@ -27,6 +27,7 @@ FSkeletonTreeVirtualBoneItem::FSkeletonTreeVirtualBoneItem(const FName& InBoneNa
 	BoneProxy = NewObject<UBoneProxy>(GetTransientPackage(), *(BoneProxyPrefix + FString::Printf(TEXT("%p"), &InSkeletonTree.Get()) + InBoneName.ToString()));
 	BoneProxy->SetFlags(RF_Transactional);
 	BoneProxy->BoneName = InBoneName;
+	BoneProxy->bIsTransformEditable = false;
 	TSharedPtr<IPersonaPreviewScene> PreviewScene = InSkeletonTree->GetPreviewScene();
 	if (PreviewScene.IsValid())
 	{
diff --git a/Engine/Source/Editor/SkeletonEditor/Public/BoneProxy.h b/Engine/Source/Editor/SkeletonEditor/Public/BoneProxy.h
--- a/Engine/Source/Editor/SkeletonEditor/Public/BoneProxy.h
+++ b/Engine/Source/Editor/SkeletonEditor/Public/BoneProxy.h
@@ -115,6 +115,9 @@ public:
 	/** Flag indicating whether this FTickableEditorObject should actually tick */
 	bool bIsTickable;
 
+	/** Flag indicating whether this bone's transform is editable */
+	bool bIsTransformEditable;
+
 	/** Method to react to retrieval of numeric values for the widget */
 	TOptional<FVector::FReal> GetNumericValue(
 		ESlateTransformComponent::Type Component,
