Tweaked skeleton tree bone weight tooltip to better reflect the weighting of a bone

Bones are included as 'weighted' if child bones are weighted

diff --git a/Engine/Source/Editor/SkeletonEditor/Private/SkeletonTreeBoneItem.cpp b/Engine/Source/Editor/SkeletonEditor/Private/SkeletonTreeBoneItem.cpp
--- a/Engine/Source/Editor/SkeletonEditor/Private/SkeletonTreeBoneItem.cpp
+++ b/Engine/Source/Editor/SkeletonEditor/Private/SkeletonTreeBoneItem.cpp
@@ -371,7 +371,7 @@ FText FSkeletonTreeBoneItem::GetBoneToolTip()
 			}
 			else
 			{
-				ToolTip = LOCTEXT( "BoneToolTipWeighted", "This bone has vertices weighted against it" );
+				ToolTip = LOCTEXT( "BoneToolTipWeighted", "This bone (or one of its children) has vertices weighted against it" );
 			}
 		}
 	}
