Added curve debugger back to the skeleton editor

Allows morph and material preview in the skeleton editor again (this was removed in the curve refactor)

diff --git a/Engine/Source/Editor/SkeletonEditor/Private/SkeletonEditor.cpp b/Engine/Source/Editor/SkeletonEditor/Private/SkeletonEditor.cpp
--- a/Engine/Source/Editor/SkeletonEditor/Private/SkeletonEditor.cpp
+++ b/Engine/Source/Editor/SkeletonEditor/Private/SkeletonEditor.cpp
@@ -39,6 +39,7 @@ const FName SkeletonEditorTabs::ViewportTab(TEXT("Viewport"));
 const FName SkeletonEditorTabs::AssetBrowserTab(TEXT("SequenceBrowser"));
 const FName SkeletonEditorTabs::AnimNotifiesTab(TEXT("SkeletonAnimNotifies"));
 const FName SkeletonEditorTabs::CurveMetadataTab(TEXT("AnimCurveMetadataEditorTab"));
+const FName SkeletonEditorTabs::CurveDebuggerTab("AnimCurveViewerTab");
 const FName SkeletonEditorTabs::AdvancedPreviewTab(TEXT("AdvancedPreviewTab"));
 const FName SkeletonEditorTabs::RetargetManagerTab(TEXT("RetargetManager"));
 const FName SkeletonEditorTabs::SlotNamesTab("SkeletonSlotNames");
diff --git a/Engine/Source/Editor/SkeletonEditor/Private/SkeletonEditor.h b/Engine/Source/Editor/SkeletonEditor/Private/SkeletonEditor.h
--- a/Engine/Source/Editor/SkeletonEditor/Private/SkeletonEditor.h
+++ b/Engine/Source/Editor/SkeletonEditor/Private/SkeletonEditor.h
@@ -36,6 +36,7 @@ namespace SkeletonEditorTabs
 	extern const FName AssetBrowserTab;
 	extern const FName AnimNotifiesTab;
 	extern const FName CurveMetadataTab;
+	extern const FName CurveDebuggerTab;
 	extern const FName AdvancedPreviewTab;
 	extern const FName RetargetManagerTab;
 	extern const FName SlotNamesTab;
diff --git a/Engine/Source/Editor/SkeletonEditor/Private/SkeletonEditorMode.cpp b/Engine/Source/Editor/SkeletonEditor/Private/SkeletonEditorMode.cpp
--- a/Engine/Source/Editor/SkeletonEditor/Private/SkeletonEditorMode.cpp
+++ b/Engine/Source/Editor/SkeletonEditor/Private/SkeletonEditorMode.cpp
@@ -39,8 +39,9 @@ FSkeletonEditorMode::FSkeletonEditorMode(TSharedRef<FWorkflowCentricApplication>
 	TabFactories.RegisterFactory(PersonaModule.CreateAssetDetailsTabFactory(InHostingApp, FOnGetAsset::CreateSP(&SkeletonEditor.Get(), &FSkeletonEditor::HandleGetAsset), FOnDetailsCreated()));
 	TabFactories.RegisterFactory(PersonaModule.CreateAnimationAssetBrowserTabFactory(InHostingApp, SkeletonEditor->GetPersonaToolkit(), FOnOpenNewAsset::CreateSP(&SkeletonEditor.Get(), &FSkeletonEditor::HandleOpenNewAsset), FOnAnimationSequenceBrowserCreated::CreateSP(&SkeletonEditor.Get(), &FSkeletonEditor::HandleAnimationSequenceBrowserCreated), true));
 	TabFactories.RegisterFactory(PersonaModule.CreateAnimAssetFindReplaceTabFactory(InHostingApp, FAnimAssetFindReplaceConfig()));
+	TabFactories.RegisterFactory(PersonaModule.CreateCurveViewerTabFactory(InHostingApp, InSkeletonTree->GetEditableSkeleton(), SkeletonEditor->GetPersonaToolkit()->GetPreviewScene(), OnObjectsSelected));
 
-	TabLayout = FTabManager::NewLayout("Standalone_SkeletonEditor_Layout_v1.4")
+	TabLayout = FTabManager::NewLayout("Standalone_SkeletonEditor_Layout_v1.5")
 		->AddArea
 		(
 			FTabManager::NewPrimaryArea()
@@ -89,6 +90,7 @@ FSkeletonEditorMode::FSkeletonEditorMode(TSharedRef<FWorkflowCentricApplication>
 						->SetSizeCoefficient(0.5f)
 						->SetHideTabWell(false)
 						->AddTab(SkeletonEditorTabs::DetailsTab, ETabState::OpenedTab)
+						->AddTab(SkeletonEditorTabs::CurveDebuggerTab, ETabState::ClosedTab)
 						->AddTab(SkeletonEditorTabs::AdvancedPreviewTab, ETabState::OpenedTab)
 						->SetForegroundTab(SkeletonEditorTabs::DetailsTab)
 					)
