Added rename support to BlendProfiles

diff --git a/Engine/Source/Editor/SkeletonEditor/Private/EditableSkeleton.cpp b/Engine/Source/Editor/SkeletonEditor/Private/EditableSkeleton.cpp
--- a/Engine/Source/Editor/SkeletonEditor/Private/EditableSkeleton.cpp
+++ b/Engine/Source/Editor/SkeletonEditor/Private/EditableSkeleton.cpp
@@ -178,6 +178,20 @@ void FEditableSkeleton::RemoveBlendProfile(UBlendProfile* InBlendProfile)
 	}
 }
 
+UBlendProfile* FEditableSkeleton::RenameBlendProfile(const FName& InBlendProfileName, const FName& InNewBlendProfileName)
+{
+	FScopedTransaction Transaction(LOCTEXT("RenameBlendProfile", "Rename Blend Profile"));
+
+	UBlendProfile* Profile = Skeleton->RenameBlendProfile(InBlendProfileName, InNewBlendProfileName);
+	
+	if (Profile == nullptr)
+	{
+		Transaction.Cancel();
+	}
+
+	return Profile;
+}
+
 void FEditableSkeleton::SetBlendProfileScale(const FName& InBlendProfileName, const FName& InBoneName, float InNewScale, bool bInRecurse)
 {
 	UBlendProfile* BlendProfile = GetBlendProfile(InBlendProfileName);
diff --git a/Engine/Source/Editor/SkeletonEditor/Private/EditableSkeleton.h b/Engine/Source/Editor/SkeletonEditor/Private/EditableSkeleton.h
--- a/Engine/Source/Editor/SkeletonEditor/Private/EditableSkeleton.h
+++ b/Engine/Source/Editor/SkeletonEditor/Private/EditableSkeleton.h
@@ -34,6 +34,7 @@ public:
 	virtual class UBlendProfile* GetBlendProfile(const FName& InBlendProfileName) override;
 	virtual class UBlendProfile* CreateNewBlendProfile(const FName& InBlendProfileName) override;
 	virtual void RemoveBlendProfile(UBlendProfile* InBlendProfile) override;
+	virtual UBlendProfile* RenameBlendProfile(const FName& InBlendProfileName, const FName& InNewBlendProfileName) override;
 	virtual void SetBlendProfileScale(const FName& InBlendProfileName, const FName& InBoneName, float InNewScale, bool bInRecurse) override;
 	virtual void SetBlendProfileMode(FName InBlendProfileName, EBlendProfileMode ProfileMode);
 	virtual USkeletalMeshSocket* DuplicateSocket(const FSelectedSocketInfo& SocketInfoToDuplicate, const FName& NewParentBoneName, USkeletalMesh* InSkeletalMesh) override;
diff --git a/Engine/Source/Editor/SkeletonEditor/Private/SSkeletonTree.cpp b/Engine/Source/Editor/SkeletonEditor/Private/SSkeletonTree.cpp
--- a/Engine/Source/Editor/SkeletonEditor/Private/SSkeletonTree.cpp
+++ b/Engine/Source/Editor/SkeletonEditor/Private/SSkeletonTree.cpp
@@ -509,6 +509,9 @@ void SSkeletonTree::BindCommands()
 		MenuActions.DeleteCurrentBlendProfile,
 		FExecuteAction::CreateSP( this, &SSkeletonTree::OnDeleteCurrentBlendProfile));
 
+	CommandList.MapAction(
+		MenuActions.RenameBlendProfile,
+		FExecuteAction::CreateSP(this, &SSkeletonTree::OnRenameBlendProfile));
 
 	PinnedCommands->BindCommandList(UICommandList.ToSharedRef());
 }
@@ -625,7 +628,18 @@ void SSkeletonTree::CreateTreeColumns()
 			})
 			.OnTextCommitted_Lambda([this](const FText& InText, ETextCommit::Type InCommitType)
 			{
-				BlendProfilePicker->OnCreateNewProfileComitted(InText, InCommitType, NewBlendProfileMode);
+				if (bIsCreateNewBlendProfile)
+				{
+					BlendProfilePicker->OnCreateNewProfileComitted(InText, InCommitType, NewBlendProfileMode);
+					bIsCreateNewBlendProfile = false;
+				}
+				else if(BlendProfilePicker->GetSelectedBlendProfileName() != NAME_None)
+				{
+					if (UBlendProfile* Profile = EditableSkeleton.Pin()->RenameBlendProfile(BlendProfilePicker->GetSelectedBlendProfileName(), FName(InText.ToString())))
+					{
+						BlendProfilePicker->SetSelectedProfile(Profile);
+					}
+				}
 			})
 			.OnVerifyTextChanged_Lambda([](const FText& InNewText, FText& OutErrorMessage) -> bool
 			{
@@ -1726,6 +1740,11 @@ void SSkeletonTree::CreateBlendProfileMenu(UToolMenu* InMenu)
 			FSlateIcon(),
 			FToolUIActionChoice(FUIAction(FExecuteAction::CreateSP(SkeletonTree->BlendProfilePicker.ToSharedRef(), &SBlendProfilePicker::OnClearSelection))));
 
+		EditSection.AddMenuEntry(
+			Actions.RenameBlendProfile,
+			FText::Format(LOCTEXT("RenameBlendProfileLabel", "Rename {0}"),
+				FText::FromName(SkeletonTree->BlendProfilePicker->GetSelectedBlendProfileName())));
+
 		EditSection.AddMenuEntry(
 			Actions.DeleteCurrentBlendProfile,
 			FText::Format(LOCTEXT("DeleteBlendProfileLabel", "Delete {0}"),
@@ -1752,6 +1771,7 @@ void SSkeletonTree::OnCreateBlendProfile(const EBlendProfileMode InMode)
 	// Activate the Header Entry Box
 	BlendProfileHeader->SetReadOnly(false);
 	BlendProfileHeader->EnterEditingMode();
+	bIsCreateNewBlendProfile = true;
 }
 
 void SSkeletonTree::OnDeleteCurrentBlendProfile()
@@ -1760,6 +1780,13 @@ void SSkeletonTree::OnDeleteCurrentBlendProfile()
 	BlendProfilePicker->OnClearSelection();
 }
 
+void SSkeletonTree::OnRenameBlendProfile()
+{
+	// Activate the Header Entry Box
+	BlendProfileHeader->SetReadOnly(false);
+	BlendProfileHeader->EnterEditingMode();
+}
+
 bool SSkeletonTree::IsBlendProfileSelected(FName ProfileName) const
 {
 	return BlendProfilePicker->GetSelectedBlendProfileName() == ProfileName;
@@ -2224,9 +2251,8 @@ void SSkeletonTree::OnBlendProfileSelected(UBlendProfile* NewProfile)
 		SkeletonTreeView->GetHeaderRow()->SetShowGeneratedColumn(ISkeletonTree::Columns::BlendProfile);
 	HandleTreeRefresh();
 
-	// When a new blend profile is created/selected - reaffirm that the header can't be edited.
-	// At this time, there is no re-naming of Blend Profiles
-	BlendProfileHeader->SetReadOnly(true);
+	// When a new blend profile is created/selected - enable edition if name != None.
+	BlendProfileHeader->SetReadOnly(BlendProfilePicker->GetSelectedBlendProfileName() == NAME_None);
 }
 
 void SSkeletonTree::RecursiveSetBlendProfileScales(float InScaleToSet)
diff --git a/Engine/Source/Editor/SkeletonEditor/Private/SSkeletonTree.h b/Engine/Source/Editor/SkeletonEditor/Private/SSkeletonTree.h
--- a/Engine/Source/Editor/SkeletonEditor/Private/SSkeletonTree.h
+++ b/Engine/Source/Editor/SkeletonEditor/Private/SSkeletonTree.h
@@ -339,6 +339,9 @@ private:
 	/** Removes the active Blend Profile */
 	void OnDeleteCurrentBlendProfile();
 
+	/** Removes the active Blend Profile */
+	void OnRenameBlendProfile();
+
 	/** Register Blend Profile Menu */
 	void RegisterBlendProfileMenu();
 
@@ -392,6 +395,9 @@ private:
 	/** Current type of blend profile no create. We shouldn't need to hold state for this, but blend profile creation is tied to its header text committed*/
 	EBlendProfileMode NewBlendProfileMode;
 
+	/** Whether we are creating or renaming a BlendProfile, so the EditBox know what to do */
+	bool bIsCreateNewBlendProfile = false;
+
 	/** Current type of bones to show */
 	EBoneFilter BoneFilter;
 
diff --git a/Engine/Source/Editor/SkeletonEditor/Private/SkeletonTreeCommands.cpp b/Engine/Source/Editor/SkeletonEditor/Private/SkeletonTreeCommands.cpp
--- a/Engine/Source/Editor/SkeletonEditor/Private/SkeletonTreeCommands.cpp
+++ b/Engine/Source/Editor/SkeletonEditor/Private/SkeletonTreeCommands.cpp
@@ -46,6 +46,7 @@ void FSkeletonTreeCommands::RegisterCommands()
 	UI_COMMAND( CreateWeightBlendProfile, "Add Weight Blend Profile", "Add a new weight blend profile for this skeleton", EUserInterfaceActionType::Button, FInputChord() );
 	UI_COMMAND( CreateBlendMask, "Add Blend Mask", "Add a new blend mask for this skeleton", EUserInterfaceActionType::Button, FInputChord() );
 	UI_COMMAND( DeleteCurrentBlendProfile, "Delete the Current Blend Profile", "Delete the active blend profile for this skeleton", EUserInterfaceActionType::Button, FInputChord() );
+	UI_COMMAND( RenameBlendProfile, "Rename the Current Blend Profile", "Rename the active blend profile for this skeleton", EUserInterfaceActionType::Button, FInputChord());
 
 }
 
diff --git a/Engine/Source/Editor/SkeletonEditor/Private/SkeletonTreeCommands.h b/Engine/Source/Editor/SkeletonEditor/Private/SkeletonTreeCommands.h
--- a/Engine/Source/Editor/SkeletonEditor/Private/SkeletonTreeCommands.h
+++ b/Engine/Source/Editor/SkeletonEditor/Private/SkeletonTreeCommands.h
@@ -119,4 +119,7 @@ public:
 
 	/** Remove the currently active blend profile */
 	TSharedPtr< FUICommandInfo > DeleteCurrentBlendProfile;
+
+	/** Rename an existing BlendProfile */
+	TSharedPtr< FUICommandInfo > RenameBlendProfile;
 };
diff --git a/Engine/Source/Editor/SkeletonEditor/Public/IEditableSkeleton.h b/Engine/Source/Editor/SkeletonEditor/Public/IEditableSkeleton.h
--- a/Engine/Source/Editor/SkeletonEditor/Public/IEditableSkeleton.h
+++ b/Engine/Source/Editor/SkeletonEditor/Public/IEditableSkeleton.h
@@ -42,6 +42,9 @@ public:
 	/** Remove the specifed blend profile */
 	virtual void RemoveBlendProfile(UBlendProfile* InBlendProfile) = 0;
 
+	/** Rename an existing blend profile */
+	virtual class UBlendProfile* RenameBlendProfile(const FName& InBlendProfileName, const FName& InNewBlendProfileName) = 0;
+
 	/** Set the blend profile scale for the specified bone */
 	virtual void SetBlendProfileScale(const FName& InBlendProfileName, const FName& InBoneName, float InNewScale, bool bInRecurse) = 0;
 
diff --git a/Engine/Source/Runtime/Engine/Classes/Animation/Skeleton.h b/Engine/Source/Runtime/Engine/Classes/Animation/Skeleton.h
--- a/Engine/Source/Runtime/Engine/Classes/Animation/Skeleton.h
+++ b/Engine/Source/Runtime/Engine/Classes/Animation/Skeleton.h
@@ -503,6 +503,9 @@ public:
 	/** Create a new blend profile with the specified name */
 	ENGINE_API UBlendProfile* CreateNewBlendProfile(const FName& InProfileName);
 
+	/** Rename an existing blend profile with the specified name. Returns the pointer if success, nullptr on failure */
+	ENGINE_API UBlendProfile* RenameBlendProfile(const FName& InProfileName, const FName& InNewProfileName);
+
 	//////////////////////////////////////////////////////////////////////////
 
 	/************************************************************************/
diff --git a/Engine/Source/Runtime/Engine/Private/Animation/Skeleton.cpp b/Engine/Source/Runtime/Engine/Private/Animation/Skeleton.cpp
--- a/Engine/Source/Runtime/Engine/Private/Animation/Skeleton.cpp
+++ b/Engine/Source/Runtime/Engine/Private/Animation/Skeleton.cpp
@@ -2053,6 +2053,24 @@ UBlendProfile* USkeleton::CreateNewBlendProfile(const FName& InProfileName)
 	return NewProfile;
 }
 
+UBlendProfile* USkeleton::RenameBlendProfile(const FName& InProfileName, const FName& InNewProfileName)
+{
+	if (!GetBlendProfile(InNewProfileName)) // we can not rename if the new name already exists
+	{
+		UBlendProfile* BlendProfile = GetBlendProfile(InProfileName);
+		if (BlendProfile)
+		{
+			Modify();
+			if (BlendProfile->Rename(*InNewProfileName.ToString(), this, REN_DontCreateRedirectors))
+			{
+				return BlendProfile;
+			}
+		}
+	}
+
+	return nullptr;
+}
+
 USkeletalMeshSocket* USkeleton::FindSocket(FName InSocketName) const
 {
 	int32 DummyIndex;
