PR #6346: Allow Sockets to be added to Virtual Bones. Enable other functionality for Virtual Bones.

diff --git a/Engine/Source/Editor/SkeletonEditor/Private/SSkeletonTree.cpp b/Engine/Source/Editor/SkeletonEditor/Private/SSkeletonTree.cpp
--- a/Engine/Source/Editor/SkeletonEditor/Private/SSkeletonTree.cpp
+++ b/Engine/Source/Editor/SkeletonEditor/Private/SSkeletonTree.cpp
@@ -757,12 +757,18 @@ TSharedPtr< SWidget > SSkeletonTree::CreateContextMenu()
 			MenuBuilder.BeginSection("SkeletonTreeBonesAction", LOCTEXT("BoneActions", "Selected Bone Actions"));
 		}
 		
-		if (BoneTreeSelection.HasSelectedOfType<FSkeletonTreeBoneItem>())
+		const bool bHasBoneSelected = BoneTreeSelection.HasSelectedOfType<FSkeletonTreeBoneItem>();
+		const bool bHasVirtualBoneSelected = BoneTreeSelection.HasSelectedOfType<FSkeletonTreeVirtualBoneItem>();
+		if (bHasBoneSelected || bHasVirtualBoneSelected)
 		{
 			MenuBuilder.AddMenuEntry(Actions.CopyBoneNames);
-			MenuBuilder.AddMenuEntry(Actions.ResetBoneTransforms);
 
-			if (BoneTreeSelection.IsSingleOfTypeSelected<FSkeletonTreeBoneItem>() && bAllowSkeletonOperations)
+			if (bHasBoneSelected)
+			{
+				MenuBuilder.AddMenuEntry(Actions.ResetBoneTransforms);
+			}
+
+			if (BoneTreeSelection.IsSingleOfTypesSelected<FSkeletonTreeBoneItem, FSkeletonTreeVirtualBoneItem>() && bAllowSkeletonOperations)
 			{
 				MenuBuilder.AddMenuEntry(Actions.AddSocket);
 				MenuBuilder.AddMenuEntry(Actions.PasteSockets);
@@ -1168,12 +1174,14 @@ void SSkeletonTree::OnCopyBoneNames()
 {
 	TArray<TSharedPtr<ISkeletonTreeItem>> SelectedItems = SkeletonTreeView->GetSelectedItems();
 	FSkeletonTreeSelection TreeSelection(SelectedItems);
-	TArray<TSharedPtr<FSkeletonTreeBoneItem>> SelectedBones = TreeSelection.GetSelectedItems<FSkeletonTreeBoneItem>();
+
+	TArray<TSharedPtr<ISkeletonTreeItem>> SelectedBones = TreeSelection.GetSelectedItemsOfTypes<FSkeletonTreeBoneItem, FSkeletonTreeVirtualBoneItem>();
+
 	if( SelectedBones.Num() > 0 )
 	{
 		bool bFirst = true;
 		FString BoneNames;
-		for (const TSharedPtr<FSkeletonTreeBoneItem>& Item : SelectedBones)
+		for (const TSharedPtr<ISkeletonTreeItem>& Item : SelectedBones)
 		{
 			FName BoneName = Item->GetRowItemName();
 			if (!bFirst)
@@ -1275,7 +1283,7 @@ void SSkeletonTree::OnPasteSockets(bool bPasteToSelectedBone)
 	FSkeletonTreeSelection TreeSelection(SelectedItems);
 
 	// Pasting sockets should only work if there is just one bone selected
-	if ( TreeSelection.IsSingleOfTypeSelected<FSkeletonTreeBoneItem>() )
+	if ( TreeSelection.IsSingleOfTypesSelected<FSkeletonTreeBoneItem, FSkeletonTreeVirtualBoneItem>())
 	{
 		FName DestBoneName = bPasteToSelectedBone ? TreeSelection.GetSingleSelectedItem()->GetRowItemName() : NAME_None;
 		USkeletalMesh* SkeletalMesh = GetPreviewScene().IsValid() ? ToRawPtr(GetPreviewScene()->GetPreviewMeshComponent()->GetSkeletalMeshAsset()) : nullptr;
@@ -1292,7 +1300,7 @@ bool SSkeletonTree::CanPasteSockets() const
 		TArray<TSharedPtr<ISkeletonTreeItem>> SelectedItems = SkeletonTreeView->GetSelectedItems();
 		FSkeletonTreeSelection TreeSelection(SelectedItems);
 
-		return TreeSelection.IsSingleOfTypeSelected<FSkeletonTreeBoneItem>();
+		return TreeSelection.IsSingleOfTypesSelected<FSkeletonTreeBoneItem, FSkeletonTreeVirtualBoneItem>();
 	}
 
 	return false;
@@ -1305,7 +1313,7 @@ void SSkeletonTree::OnAddSocket()
 	FSkeletonTreeSelection TreeSelection(SelectedItems);
 
 	// Can only add a socket to one bone
-	if (TreeSelection.IsSingleOfTypeSelected<FSkeletonTreeBoneItem>())
+	if (TreeSelection.IsSingleOfTypesSelected<FSkeletonTreeBoneItem, FSkeletonTreeVirtualBoneItem>())
 	{
 		FName BoneName = TreeSelection.GetSingleSelectedItem()->GetRowItemName();
 		USkeletalMeshSocket* NewSocket = GetEditableSkeletonInternal()->HandleAddSocket(BoneName);
diff --git a/Engine/Source/Editor/SkeletonEditor/Private/SkeletonTreeBuilder.cpp b/Engine/Source/Editor/SkeletonEditor/Private/SkeletonTreeBuilder.cpp
--- a/Engine/Source/Editor/SkeletonEditor/Private/SkeletonTreeBuilder.cpp
+++ b/Engine/Source/Editor/SkeletonEditor/Private/SkeletonTreeBuilder.cpp
@@ -94,6 +94,11 @@ void FSkeletonTreeBuilder::Build(FSkeletonTreeBuilderOutput& Output)
 		AddBones(Output);
 	}
 
+	if (BuilderArgs.bShowVirtualBones)
+	{
+		AddVirtualBones(Output);
+	}
+
 	if (BuilderArgs.bShowSockets)
 	{
 		AddSockets(Output);
@@ -103,11 +108,6 @@ void FSkeletonTreeBuilder::Build(FSkeletonTreeBuilderOutput& Output)
 	{
 		AddAttachedAssets(Output);
 	}
-
-	if (BuilderArgs.bShowVirtualBones)
-	{
-		AddVirtualBones(Output);
-	}
 }
 
 void FSkeletonTreeBuilder::Filter(const FSkeletonTreeFilterArgs& InArgs, const TArray<TSharedPtr<ISkeletonTreeItem>>& InItems, TArray<TSharedPtr<ISkeletonTreeItem>>& OutFilteredItems)
@@ -355,7 +355,7 @@ void FSkeletonTreeBuilder::AddSocketsFromData(const TArray< USkeletalMeshSocket*
 			}
 		}
 
-		Output.Add(CreateSocketTreeItem(Socket, ParentType, bIsCustomized), Socket->BoneName, FSkeletonTreeBoneItem::GetTypeId(), /*bAddToHead*/ true);
+		Output.Add(CreateSocketTreeItem(Socket, ParentType, bIsCustomized), Socket->BoneName, { FSkeletonTreeBoneItem::GetTypeId(), FSkeletonTreeVirtualBoneItem::GetTypeId() }, /*bAddToHead*/ true);
 	}
 }
 
diff --git a/Engine/Source/Editor/SkeletonEditor/Private/SkeletonTreeVirtualBoneItem.cpp b/Engine/Source/Editor/SkeletonEditor/Private/SkeletonTreeVirtualBoneItem.cpp
--- a/Engine/Source/Editor/SkeletonEditor/Private/SkeletonTreeVirtualBoneItem.cpp
+++ b/Engine/Source/Editor/SkeletonEditor/Private/SkeletonTreeVirtualBoneItem.cpp
@@ -13,6 +13,7 @@
 #include "Animation/DebugSkelMeshComponent.h"
 #include "Animation/BlendProfile.h"
 #include "UObject/Package.h"
+#include "SocketDragDropOp.h"
 #include "Editor.h"
 #include "ScopedTransaction.h"
 
@@ -244,6 +245,62 @@ void FSkeletonTreeVirtualBoneItem::OnItemDoubleClicked()
 	OnRenameRequested.ExecuteIfBound();
 }
 
+void FSkeletonTreeVirtualBoneItem::HandleDragEnter(const FDragDropEvent& DragDropEvent)
+{
+	TSharedPtr<FSocketDragDropOp> DragConnectionOp = DragDropEvent.GetOperationAs<FSocketDragDropOp>();
+
+	// Is someone dragging a socket onto a bone?
+	if (DragConnectionOp.IsValid())
+	{
+		if (BoneName != DragConnectionOp->GetSocketInfo().Socket->BoneName)
+		{
+			// The socket can be dropped here if we're a bone and NOT the socket's existing parent
+			DragConnectionOp->SetIcon(FAppStyle::GetBrush(TEXT("Graph.ConnectorFeedback.Ok")));
+		}
+		else if (DragConnectionOp->IsAltDrag())
+		{
+			// For Alt-Drag, dropping onto the existing parent is fine, as we're going to copy, not move the socket
+			DragConnectionOp->SetIcon(FAppStyle::GetBrush(TEXT("Graph.ConnectorFeedback.Ok")));
+		}
+	}
+}
+
+void FSkeletonTreeVirtualBoneItem::HandleDragLeave(const FDragDropEvent& DragDropEvent)
+{
+	TSharedPtr<FSocketDragDropOp> DragConnectionOp = DragDropEvent.GetOperationAs<FSocketDragDropOp>();
+	if (DragConnectionOp.IsValid())
+	{
+		// Reset the drag/drop icon when leaving this row
+		DragConnectionOp->SetIcon(FAppStyle::GetBrush(TEXT("Graph.ConnectorFeedback.Error")));
+	}
+}
+
+FReply FSkeletonTreeVirtualBoneItem::HandleDrop(const FDragDropEvent& DragDropEvent)
+{
+	TSharedPtr<FSocketDragDropOp> DragConnectionOp = DragDropEvent.GetOperationAs<FSocketDragDropOp>();
+	if (DragConnectionOp.IsValid())
+	{
+		FSelectedSocketInfo SocketInfo = DragConnectionOp->GetSocketInfo();
+
+		if (DragConnectionOp->IsAltDrag())
+		{
+			// In an alt-drag, the socket can be dropped on any bone
+			// (including its existing parent) to create a uniquely named copy
+			GetSkeletonTree()->DuplicateAndSelectSocket(SocketInfo, BoneName);
+		}
+		else if (BoneName != SocketInfo.Socket->BoneName)
+		{
+			// The socket can be dropped here if we're a bone and NOT the socket's existing parent
+			USkeletalMesh* SkeletalMesh = GetSkeletonTree()->GetPreviewScene().IsValid() ? GetSkeletonTree()->GetPreviewScene()->GetPreviewMeshComponent()->GetSkeletalMeshAsset() : nullptr;
+			GetEditableSkeleton()->SetSocketParent(SocketInfo.Socket->SocketName, BoneName, SkeletalMesh);
+
+			return FReply::Handled();
+		}
+	}
+
+	return FReply::Unhandled();
+}
+
 void FSkeletonTreeVirtualBoneItem::RequestRename()
 {
 	OnRenameRequested.ExecuteIfBound();
diff --git a/Engine/Source/Editor/SkeletonEditor/Private/SkeletonTreeVirtualBoneItem.h b/Engine/Source/Editor/SkeletonEditor/Private/SkeletonTreeVirtualBoneItem.h
--- a/Engine/Source/Editor/SkeletonEditor/Private/SkeletonTreeVirtualBoneItem.h
+++ b/Engine/Source/Editor/SkeletonEditor/Private/SkeletonTreeVirtualBoneItem.h
@@ -31,6 +31,9 @@ public:
 	virtual bool CanRenameItem() const override { return true; }
 	virtual void RequestRename() override;
 	virtual void OnItemDoubleClicked() override;
+	virtual void HandleDragEnter(const FDragDropEvent& DragDropEvent) override;
+	virtual void HandleDragLeave(const FDragDropEvent& DragDropEvent) override;
+	virtual FReply HandleDrop(const FDragDropEvent& DragDropEvent) override;
 	virtual UObject* GetObject() const override { return BoneProxy; }
 
 	/** FGCObject interface */
diff --git a/Engine/Source/Editor/SkeletonEditor/Public/SkeletonTreeSelection.h b/Engine/Source/Editor/SkeletonEditor/Public/SkeletonTreeSelection.h
--- a/Engine/Source/Editor/SkeletonEditor/Public/SkeletonTreeSelection.h
+++ b/Engine/Source/Editor/SkeletonEditor/Public/SkeletonTreeSelection.h
@@ -3,6 +3,7 @@
 #pragma once
 
 #include "CoreMinimal.h"
+#include "Algo/AnyOf.h"
 #include "ISkeletonTreeItem.h"
 
 /** Helper for processing tree selections */
@@ -42,6 +43,17 @@ public:
 		return false;
 	}
 
+	template<typename... ItemTypes>
+	bool IsSingleOfTypesSelected() const
+	{
+		if (IsSingleItemSelected())
+		{
+			const bool bTypeResults[] = { false, (SelectedItems[0]->IsOfType<ItemTypes>())... };
+			return Algo::AnyOf(bTypeResults);
+		}
+		return false;
+	}
+
 	TSharedPtr<ISkeletonTreeItem> GetSingleSelectedItem()
 	{
 		check(IsSingleItemSelected());
@@ -62,6 +74,21 @@ public:
 		return false;
 	}
 
+	template<typename... ItemTypes>
+	bool HasSelectedOfTypes() const
+	{
+		for (const TSharedPtr<ISkeletonTreeItem>& SelectedItem : SelectedItems)
+		{
+			const bool bTypeResults[] = { false, (SelectedItem->IsOfType<ItemTypes>())... };
+			if (Algo::AnyOf(bTypeResults))
+			{
+				return true;
+			}
+		}
+
+		return false;
+	}
+
 	template<typename ItemType>
 	TArray<TSharedPtr<ItemType>> GetSelectedItems() const
 	{
@@ -76,6 +103,21 @@ public:
 		return Items;
 	}
 
+	template<typename... ItemTypes>
+	TArray<TSharedPtr<ISkeletonTreeItem>> GetSelectedItemsOfTypes() const
+	{
+		TArray<TSharedPtr<ISkeletonTreeItem>> Items;
+		for (const TSharedPtr<ISkeletonTreeItem>& SelectedItem : SelectedItems)
+		{
+			const bool bTypeResults[] = { false, (SelectedItem->IsOfType<ItemTypes>())... };
+			if (Algo::AnyOf(bTypeResults))
+			{
+				Items.Add(SelectedItem);
+			}
+		}
+		return Items;
+	}
+
 	TArray<TSharedPtr<ISkeletonTreeItem>> GetSelectedItemsByTypeId(const FName& InTypeId) const
 	{
 		TArray<TSharedPtr<ISkeletonTreeItem>> Items;
