Skeleton Editor: Fix bone proxy to properly transact on interactive/slider changes and paste. Fix an issue with "reset to default" double-stacking transforms instead of resetting.

diff --git a/Engine/Source/Editor/SkeletonEditor/Private/BoneProxy.cpp b/Engine/Source/Editor/SkeletonEditor/Private/BoneProxy.cpp
--- a/Engine/Source/Editor/SkeletonEditor/Private/BoneProxy.cpp
+++ b/Engine/Source/Editor/SkeletonEditor/Private/BoneProxy.cpp
@@ -4,6 +4,7 @@
 #include "IPersonaPreviewScene.h"
 #include "BoneControllers/AnimNode_ModifyBone.h"
 #include "AnimPreviewInstance.h"
+#include "Editor.h"
 #include "IPersonaToolkit.h"
 #include "SAdvancedTransformInputBox.h"
 #include "ScopedTransaction.h"
@@ -54,7 +55,7 @@ void UBoneProxy::Tick(float DeltaTime)
 		{
 			if (Component->GetSkeletalMeshAsset() && Component->GetSkeletalMeshAsset()->IsCompiling())
 			{
-				//We do not want to tick if the skeletalmesh is inside a compilation
+				//We do not want to tick if the skeletal mesh is inside a compilation
 				return;
 			}
 
@@ -179,6 +180,9 @@ TOptional<FVector::FReal> UBoneProxy::GetNumericValue(
 
 	switch(TransformType)
 	{
+		case TransformType_Bone:
+			// Pointers are already set.
+			break;
 		case TransformType_Reference:
 		{
 			LocationPtr = &ReferenceLocation;
@@ -233,13 +237,12 @@ TOptional<FVector::FReal> UBoneProxy::GetMultiNumericValue(
 	return FirstValue;
 }
 
-void UBoneProxy::OnNumericValueCommitted(
+void UBoneProxy::OnNumericValueChanged(
 	ESlateTransformComponent::Type Component,
 	ESlateRotationRepresentation::Type Representation,
 	ESlateTransformSubComponent::Type SubComponent,
-	FVector::FReal Value, ETextCommit::Type CommitType,
-	ETransformType TransformType,
-	bool bIsCommit)
+	FVector::FReal Value,
+	ETransformType TransformType)
 {
 	if(TransformType != TransformType_Bone || !bIsTransformEditable)
 	{
@@ -250,7 +253,7 @@ void UBoneProxy::OnNumericValueCommitted(
 	{
 		case ESlateTransformComponent::Location:
 		{
-			OnPreEditChange(GET_MEMBER_NAME_CHECKED(UBoneProxy, Location), bIsCommit);
+			OnPreEditChange(GET_MEMBER_NAME_CHECKED(UBoneProxy, Location));
 				
 			switch(SubComponent)
 			{
@@ -269,14 +272,16 @@ void UBoneProxy::OnNumericValueCommitted(
 					Location.Z = Value;
 					break;
 				}
+				default:
+					checkNoEntry();
 			}
 
-			OnPostEditChangeProperty(GET_MEMBER_NAME_CHECKED(UBoneProxy, Location), bIsCommit);
+			OnPostEditChangeProperty(GET_MEMBER_NAME_CHECKED(UBoneProxy, Location));
 			break;
 		}
 		case ESlateTransformComponent::Rotation:
 		{
-			OnPreEditChange(GET_MEMBER_NAME_CHECKED(UBoneProxy, Rotation), bIsCommit);
+			OnPreEditChange(GET_MEMBER_NAME_CHECKED(UBoneProxy, Rotation));
 				
 			switch(SubComponent)
 			{
@@ -295,14 +300,16 @@ void UBoneProxy::OnNumericValueCommitted(
 					Rotation.Yaw = Value;
 					break;
 				}
+				default:
+					checkNoEntry();
 			}
 
-			OnPostEditChangeProperty(GET_MEMBER_NAME_CHECKED(UBoneProxy, Rotation), bIsCommit);
+			OnPostEditChangeProperty(GET_MEMBER_NAME_CHECKED(UBoneProxy, Rotation));
 			break;
 		}
 		case ESlateTransformComponent::Scale:
 		{
-			OnPreEditChange(GET_MEMBER_NAME_CHECKED(UBoneProxy, Scale), bIsCommit);
+			OnPreEditChange(GET_MEMBER_NAME_CHECKED(UBoneProxy, Scale));
 
 			switch(SubComponent)
 			{
@@ -321,27 +328,87 @@ void UBoneProxy::OnNumericValueCommitted(
 					Scale.Z = Value;
 					break;
 				}
+				default:
+					checkNoEntry();
 			}
 
-			OnPostEditChangeProperty(GET_MEMBER_NAME_CHECKED(UBoneProxy, Scale), bIsCommit);
+			OnPostEditChangeProperty(GET_MEMBER_NAME_CHECKED(UBoneProxy, Scale));
 			break;
 		}
+		default:
+			checkNoEntry();
+	}
+}
+
+
+void UBoneProxy::OnSliderMovementStateChanged(
+	ESlateTransformComponent::Type Component,
+	ESlateRotationRepresentation::Type Representation,
+	ESlateTransformSubComponent::Type SubComponent,
+	FVector::FReal Value, 
+	ESliderMovementState SliderMovementState,
+	TArrayView<UBoneProxy*> BoneProxies
+	)
+{
+	// If doing slider movement, we have to start the transaction at slider move begin so that
+	// we can make a copy of the state before value changes begin. The slider end transition
+	// occurs after value commit has happened, so that's where we end the transaction.
+	if (SliderMovementState == ESliderMovementState::Begin)
+	{
+		BeginSetValueTransaction(Component, BoneProxies);
+	}
+	else // == ESliderMovementState::End
+	{
+		EndTransaction();
+	}
+	
+	for(UBoneProxy* BoneProxy : BoneProxies)
+	{
+		BoneProxy->bManipulating = SliderMovementState == ESliderMovementState::Begin;
 	}
 }
 
-void UBoneProxy::OnMultiNumericValueCommitted(
+
+void UBoneProxy::OnMultiNumericValueChanged(
 	ESlateTransformComponent::Type Component,
 	ESlateRotationRepresentation::Type Representation,
 	ESlateTransformSubComponent::Type SubComponent,
 	FVector::FReal Value,
-	ETextCommit::Type CommitType,
+	ETextCommit::Type InCommitType,
+	bool bInTransactional,
 	ETransformType TransformType,
-	TArrayView<UBoneProxy*> BoneProxies,
-	bool bIsCommit)
+	TArrayView<UBoneProxy*> BoneProxies)
 {
+	if(TransformType != TransformType_Bone)
+	{
+		return;
+	}
+
+	// Are we in a slider movement? In that case, transaction is handled by the OnSliderMovementStateChanged
+	// callback, since that completely brackets calls to OnMultiNumericValueChanged for change/commit values.
+	bool bInSliderMovement = false;
+	for(UBoneProxy* BoneProxy : BoneProxies)
+	{
+		if (BoneProxy->bManipulating)
+		{
+			bInSliderMovement = true;
+			break;
+		}
+	}
+
+	if (bInTransactional && !bInSliderMovement)
+	{
+		BeginSetValueTransaction(Component, BoneProxies);
+	}
+	
 	for(UBoneProxy* BoneProxy : BoneProxies)
 	{
-		BoneProxy->OnNumericValueCommitted(Component, Representation, SubComponent, Value, CommitType, TransformType, bIsCommit);
+		BoneProxy->OnNumericValueChanged(Component, Representation, SubComponent, Value, TransformType);
+	}
+	
+	if (bInTransactional && !bInSliderMovement)
+	{
+		EndTransaction();
 	}
 }
 
@@ -374,15 +441,16 @@ bool UBoneProxy::DiffersFromDefault(ESlateTransformComponent::Type Component, ET
 	return false;
 }
 
+
 void UBoneProxy::ResetToDefault(ESlateTransformComponent::Type InComponent, ETransformType TransformType)
 {
 	if(TransformType == TransformType_Bone && bIsTransformEditable)
 	{
-		if (UDebugSkelMeshComponent* Component = SkelMeshComponent.Get())
+		if (const UDebugSkelMeshComponent* Component = SkelMeshComponent.Get())
 		{
 			if (Component->PreviewInstance && Component->AnimScriptInstance == Component->PreviewInstance)
 			{
-				int32 BoneIndex = Component->GetBoneIndex(BoneName);
+				const int32 BoneIndex = Component->GetBoneIndex(BoneName);
 				if (BoneIndex != INDEX_NONE && BoneIndex < Component->GetNumComponentSpaceTransforms())
 				{
 					Component->PreviewInstance->SetFlags(RF_Transactional);
@@ -394,31 +462,31 @@ void UBoneProxy::ResetToDefault(ESlateTransformComponent::Type InComponent, ETra
 					{
 						case ESlateTransformComponent::Location:
 						{
-							ModifyBone.Translation = ReferenceLocation;
+							ModifyBone.Translation = FVector::ZeroVector;
 							break;
 						}
 						case ESlateTransformComponent::Rotation:
 						{
-							ModifyBone.Rotation = ReferenceRotation;
+							ModifyBone.Rotation = FRotator::ZeroRotator;
 							break;
 						}
 						case ESlateTransformComponent::Scale:
 						{
-							ModifyBone.Scale = ReferenceScale;
+							ModifyBone.Scale = FVector::OneVector;
 							break;
 						}
 						default:
 						{
-							ModifyBone.Translation = ReferenceLocation;
-							ModifyBone.Rotation = ReferenceRotation;
-							ModifyBone.Scale = ReferenceScale;
+							ModifyBone.Translation = FVector::ZeroVector;
+							ModifyBone.Rotation = FRotator::ZeroRotator;
+							ModifyBone.Scale = FVector::OneVector;
 							break;
 						}
 					}
 
-					if(ModifyBone.Translation.Equals(ReferenceLocation) &&
-						ModifyBone.Rotation.Equals(ReferenceRotation) &&
-						ModifyBone.Scale.Equals(ReferenceScale))
+					if(ModifyBone.Translation.Equals(FVector::ZeroVector) &&
+						ModifyBone.Rotation.Equals(FRotator::ZeroRotator) &&
+						ModifyBone.Scale.Equals(FVector::OneVector))
 					{
 						Component->PreviewInstance->RemoveBoneModification(BoneName);
 					}
@@ -428,16 +496,50 @@ void UBoneProxy::ResetToDefault(ESlateTransformComponent::Type InComponent, ETra
 	}
 }
 
+void UBoneProxy::BeginSetValueTransaction(ESlateTransformComponent::Type InComponent, TArrayView<UBoneProxy*> BoneProxies)
+{
+	FText TransactionScopeText;
+	switch (InComponent)
+	{
+	case ESlateTransformComponent::Location:
+		TransactionScopeText = LOCTEXT("SetLocation", "Set Location");
+		break;
+	case ESlateTransformComponent::Rotation:
+		TransactionScopeText = LOCTEXT("SetRotation", "Set Rotation");
+		break;
+	case ESlateTransformComponent::Scale:
+		TransactionScopeText = LOCTEXT("SetScale", "Set Scale");
+		break;
+	default:
+		checkNoEntry();
+	}
+
+	GEditor->BeginTransaction(TransactionScopeText);
+
+	for (UBoneProxy* BoneProxy: BoneProxies)
+	{
+		if (UDebugSkelMeshComponent* Component = BoneProxy->SkelMeshComponent.Get())
+		{
+			if (Component->PreviewInstance && Component->AnimScriptInstance == Component->PreviewInstance)
+			{
+				Component->PreviewInstance->SetFlags(RF_Transactional);
+				Component->PreviewInstance->Modify();
+			}
+		}
+	}
+}
+
+void UBoneProxy::EndTransaction()
+{
+	GEditor->EndTransaction();
+}
+
 void UBoneProxy::PreEditChange(FEditPropertyChain& PropertyAboutToChange)
 {
 	if (UDebugSkelMeshComponent* Component = SkelMeshComponent.Get())
 	{
 		if (Component->PreviewInstance && Component->AnimScriptInstance == Component->PreviewInstance)
 		{
-			bManipulating = true;
-
-			Component->PreviewInstance->Modify();
-
 			if (PropertyAboutToChange.GetActiveMemberNode()->GetValue()->GetFName() == GET_MEMBER_NAME_CHECKED(UBoneProxy, Location))
 			{
 				PreviousLocation = Location;
@@ -462,15 +564,11 @@ void UBoneProxy::PostEditChangeProperty(FPropertyChangedEvent& PropertyChangedEv
 		{
 			if (Component->PreviewInstance && Component->AnimScriptInstance == Component->PreviewInstance)
 			{
-				bManipulating = (PropertyChangedEvent.ChangeType == EPropertyChangeType::Interactive);
-
 				int32 BoneIndex = Component->GetBoneIndex(BoneName);
 				if (BoneIndex != INDEX_NONE && BoneIndex < Component->GetNumComponentSpaceTransforms())
 				{
 					FTransform BoneTransform = Component->GetBoneTransform(BoneIndex);
 					FMatrix BoneLocalCoordSystem = Component->GetBoneTransform(BoneIndex).ToMatrixNoScale().RemoveTranslation();
-					Component->PreviewInstance->SetFlags(RF_Transactional);
-					Component->PreviewInstance->Modify();
 					FAnimNode_ModifyBone& ModifyBone = Component->PreviewInstance->ModifyBone(BoneName);
 					FTransform ModifyBoneTransform(ModifyBone.Rotation, ModifyBone.Translation, ModifyBone.Scale);
 					FTransform BaseTransform = BoneTransform.GetRelativeTransformReverse(ModifyBoneTransform);
@@ -528,7 +626,7 @@ void UBoneProxy::PostEditChangeProperty(FPropertyChangedEvent& PropertyChangedEv
 	}
 }
 
-void UBoneProxy::OnPreEditChange(FName PropertyName, bool bIsCommit)
+void UBoneProxy::OnPreEditChange(FName PropertyName)
 {
 	FEditPropertyChain PropertyChain;
 	FProperty* Property = FindFProperty<FProperty>(GetClass(), PropertyName);
@@ -537,10 +635,10 @@ void UBoneProxy::OnPreEditChange(FName PropertyName, bool bIsCommit)
 	PreEditChange(PropertyChain);
 }
 
-void UBoneProxy::OnPostEditChangeProperty(FName PropertyName, bool bIsCommit)
+void UBoneProxy::OnPostEditChangeProperty(FName PropertyName)
 {
 	FProperty* Property = FindFProperty<FProperty>(GetClass(), PropertyName);
-	FPropertyChangedEvent ChangedEvent(Property, bIsCommit ? EPropertyChangeType::Unspecified : EPropertyChangeType::Interactive);
+	FPropertyChangedEvent ChangedEvent(Property, bManipulating ? EPropertyChangeType::Interactive : EPropertyChangeType::ValueSet);
 	PostEditChangeProperty(ChangedEvent);
 }
 
diff --git a/Engine/Source/Editor/SkeletonEditor/Private/BoneProxyDetailsCustomization.cpp b/Engine/Source/Editor/SkeletonEditor/Private/BoneProxyDetailsCustomization.cpp
--- a/Engine/Source/Editor/SkeletonEditor/Private/BoneProxyDetailsCustomization.cpp
+++ b/Engine/Source/Editor/SkeletonEditor/Private/BoneProxyDetailsCustomization.cpp
@@ -387,7 +387,6 @@ void FBoneProxyDetailsCustomization::CustomizeDetails(IDetailLayoutBuilder& Deta
 			static const FName LocationPropertyName = GET_MEMBER_NAME_CHECKED(UBoneProxy, Location);
 			static const FName RotationPropertyName = GET_MEMBER_NAME_CHECKED(UBoneProxy, Rotation);
 			static const FName ScalePropertyName = GET_MEMBER_NAME_CHECKED(UBoneProxy, Scale);
-			static constexpr bool bIsCommit = true;
 
 			for (UBoneProxy* BoneProxy : BoneProxiesView)
 			{
@@ -397,6 +396,9 @@ void FBoneProxyDetailsCustomization::CustomizeDetails(IDetailLayoutBuilder& Deta
 				}
 				if (const UDebugSkelMeshComponent* Component = BoneProxy->SkelMeshComponent.Get())
 				{
+					Component->PreviewInstance->SetFlags(RF_Transactional);
+					Component->PreviewInstance->Modify();
+					
 					switch(InComponent)
 					{
 						case ESlateTransformComponent::Location:
@@ -404,9 +406,9 @@ void FBoneProxyDetailsCustomization::CustomizeDetails(IDetailLayoutBuilder& Deta
 							FVector Data = BoneProxy->Location;
 							if (GetDataFromContent(Content, Data))
 							{
-								BoneProxy->OnPreEditChange(LocationPropertyName, bIsCommit);
+								BoneProxy->OnPreEditChange(LocationPropertyName);
 								BoneProxy->Location = Data;
-								BoneProxy->OnPostEditChangeProperty(LocationPropertyName, bIsCommit);
+								BoneProxy->OnPostEditChangeProperty(LocationPropertyName);
 							}
 							break;
 						}
@@ -415,9 +417,9 @@ void FBoneProxyDetailsCustomization::CustomizeDetails(IDetailLayoutBuilder& Deta
 							FRotator Data = BoneProxy->Rotation;
 							if (GetDataFromContent(Content, Data))
 							{
-								BoneProxy->OnPreEditChange(RotationPropertyName, bIsCommit);
+								BoneProxy->OnPreEditChange(RotationPropertyName);
 								BoneProxy->Rotation = Data;
-								BoneProxy->OnPostEditChangeProperty(RotationPropertyName, bIsCommit);
+								BoneProxy->OnPostEditChangeProperty(RotationPropertyName);
 							}
 							break;
 						}
@@ -426,9 +428,9 @@ void FBoneProxyDetailsCustomization::CustomizeDetails(IDetailLayoutBuilder& Deta
 							FVector Data = BoneProxy->Scale;
 							if (GetDataFromContent(Content, Data))
 							{
-								BoneProxy->OnPreEditChange(ScalePropertyName, bIsCommit);
+								BoneProxy->OnPreEditChange(ScalePropertyName);
 								BoneProxy->Scale = Data;
-								BoneProxy->OnPostEditChangeProperty(ScalePropertyName, bIsCommit);
+								BoneProxy->OnPostEditChangeProperty(ScalePropertyName);
 							}
 							break;
 						}
@@ -438,17 +440,17 @@ void FBoneProxyDetailsCustomization::CustomizeDetails(IDetailLayoutBuilder& Deta
 							FEulerTransform Data = FEulerTransform::Identity;
 							if (GetDataFromContent(Content, Data))
 							{
-								BoneProxy->OnPreEditChange(LocationPropertyName, bIsCommit);
+								BoneProxy->OnPreEditChange(LocationPropertyName);
 								BoneProxy->Location = Data.GetLocation();
-								BoneProxy->OnPostEditChangeProperty(LocationPropertyName, bIsCommit);
+								BoneProxy->OnPostEditChangeProperty(LocationPropertyName);
 								
-								BoneProxy->OnPreEditChange(RotationPropertyName, bIsCommit);
+								BoneProxy->OnPreEditChange(RotationPropertyName);
 								BoneProxy->Rotation = Data.Rotator();
-								BoneProxy->OnPostEditChangeProperty(RotationPropertyName, bIsCommit);
+								BoneProxy->OnPostEditChangeProperty(RotationPropertyName);
 								
-								BoneProxy->OnPreEditChange(ScalePropertyName, bIsCommit);
+								BoneProxy->OnPreEditChange(ScalePropertyName);
 								BoneProxy->Scale = Data.GetScale3D();
-								BoneProxy->OnPostEditChangeProperty(ScalePropertyName, bIsCommit);
+								BoneProxy->OnPostEditChangeProperty(ScalePropertyName);
 							}
 							break;
 						}
@@ -469,6 +471,8 @@ void FBoneProxyDetailsCustomization::CustomizeDetails(IDetailLayoutBuilder& Deta
 		})
 		.OnResetToDefault_Lambda([TransformType, BoneProxiesView](ESlateTransformComponent::Type InComponent)
 		{
+			FScopedTransaction Transaction(LOCTEXT("ResetToDefault", "Reset to Default"));
+			
 			for(UBoneProxy* BoneProxy : BoneProxiesView)
 			{
 				BoneProxy->ResetToDefault(InComponent, TransformType);
@@ -482,8 +486,12 @@ void FBoneProxyDetailsCustomization::CustomizeDetails(IDetailLayoutBuilder& Deta
 
 		if(bIsEditingEnabled)
 		{
-			TransformWidgetArgs.OnNumericValueChanged_Static(&UBoneProxy::OnMultiNumericValueCommitted, ETextCommit::Default, TransformType, BoneProxiesView, false);
-			TransformWidgetArgs.OnNumericValueCommitted_Static (&UBoneProxy::OnMultiNumericValueCommitted, TransformType, BoneProxiesView, true);
+			constexpr bool bNonTransactional = false;
+			constexpr bool bTransactional = true;
+			TransformWidgetArgs.OnBeginSliderMovement_Static(&UBoneProxy::OnSliderMovementStateChanged, 0.0, UBoneProxy::ESliderMovementState::Begin, BoneProxiesView);
+			TransformWidgetArgs.OnEndSliderMovement_Static(&UBoneProxy::OnSliderMovementStateChanged, UBoneProxy::ESliderMovementState::End, BoneProxiesView);
+			TransformWidgetArgs.OnNumericValueChanged_Static(&UBoneProxy::OnMultiNumericValueChanged, ETextCommit::Default, bNonTransactional, TransformType, BoneProxiesView);
+			TransformWidgetArgs.OnNumericValueCommitted_Static (&UBoneProxy::OnMultiNumericValueChanged, bTransactional, TransformType, BoneProxiesView);
 		}
 
 		SAdvancedTransformInputBox<FEulerTransform>::ConstructGroupedTransformRows(
diff --git a/Engine/Source/Editor/SkeletonEditor/Public/BoneProxy.h b/Engine/Source/Editor/SkeletonEditor/Public/BoneProxy.h
--- a/Engine/Source/Editor/SkeletonEditor/Public/BoneProxy.h
+++ b/Engine/Source/Editor/SkeletonEditor/Public/BoneProxy.h
@@ -36,13 +36,20 @@ public:
 		TransformType_Mesh
 	};
 
+	enum class ESliderMovementState
+	{
+		Begin,
+		End
+	};
+
+	
 	/** UObject interface */
 	virtual void PreEditChange(FEditPropertyChain& PropertyAboutToChange) override;
 	virtual void PostEditChangeProperty(FPropertyChangedEvent& PropertyChangedEvent) override;
 
 	// overload based on a property name
-	void OnPreEditChange(FName PropertyName, bool bIsCommit);
-	void OnPostEditChangeProperty(FName PropertyName, bool bIsCommit);
+	void OnPreEditChange(FName PropertyName);
+	void OnPostEditChangeProperty(FName PropertyName);
 	
 	/** FTickableEditorObject interface */
 	virtual void Tick(float DeltaTime) override;
@@ -134,29 +141,39 @@ public:
 		TArrayView<UBoneProxy*> BoneProxies);
 
 	/** Method to react to changes of numeric values in the widget */
-	void OnNumericValueCommitted(
+	void OnNumericValueChanged(
 		ESlateTransformComponent::Type Component,
 		ESlateRotationRepresentation::Type Representation,
 		ESlateTransformSubComponent::Type SubComponent,
 		FVector::FReal Value,
-		ETextCommit::Type CommitType,
-		ETransformType TransformType,
-		bool bIsCommit);
+		ETransformType TransformType);
+
+	static void OnSliderMovementStateChanged(
+		ESlateTransformComponent::Type Component,
+		ESlateRotationRepresentation::Type Representation,
+		ESlateTransformSubComponent::Type SubComponent,
+		FVector::FReal Value,
+		ESliderMovementState SliderMovementState,
+		TArrayView<UBoneProxy*> BoneProxies);
 
 	/** Method to react to changes of numeric values in the widget */
-	static void OnMultiNumericValueCommitted(
+	static void OnMultiNumericValueChanged(
 		ESlateTransformComponent::Type Component,
 		ESlateRotationRepresentation::Type Representation,
 		ESlateTransformSubComponent::Type SubComponent,
 		FVector::FReal Value,
-		ETextCommit::Type CommitType,
+		ETextCommit::Type InCommitType,
+		bool bInTransactional,
 		ETransformType TransformType,
-		TArrayView<UBoneProxy*> BoneProxies,
-		bool bIsCommit);
+		TArrayView<UBoneProxy*> BoneProxies);
 
 	// Returns true if a given transform component differs from the reference
 	bool DiffersFromDefault(ESlateTransformComponent::Type Component, ETransformType TransformType) const;
 
 	// Resets the bone transform's component to its reference
 	void ResetToDefault(ESlateTransformComponent::Type InComponent, ETransformType TransformType);
+
+private:
+	static void BeginSetValueTransaction(ESlateTransformComponent::Type InComponent, TArrayView<UBoneProxy*> BoneProxies);
+	static void EndTransaction();
 };
