Disable 'Add Socket' action in skeleton tree when a bone is not selected

diff --git a/Engine/Source/Editor/SkeletonEditor/Private/SSkeletonTree.cpp b/Engine/Source/Editor/SkeletonEditor/Private/SSkeletonTree.cpp
--- a/Engine/Source/Editor/SkeletonEditor/Private/SSkeletonTree.cpp
+++ b/Engine/Source/Editor/SkeletonEditor/Private/SSkeletonTree.cpp
@@ -2135,7 +2135,9 @@ bool SSkeletonTree::IsAddingSocketsAllowed() const
 		SocketFilter == ESocketFilter::Active ||
 		SocketFilter == ESocketFilter::All )
 	{
-		return true;
+		TArray<TSharedPtr<ISkeletonTreeItem>> SelectedItems = SkeletonTreeView->GetSelectedItems();
+		FSkeletonTreeSelection TreeSelection(SelectedItems);
+		return TreeSelection.IsSingleOfTypesSelected<FSkeletonTreeBoneItem, FSkeletonTreeVirtualBoneItem>();
 	}
 
 	return false;
