Skeleton Tree: Option to show bone numbers

diff --git a/Engine/Source/Editor/SkeletonEditor/Private/SSkeletonTree.cpp b/Engine/Source/Editor/SkeletonEditor/Private/SSkeletonTree.cpp
--- a/Engine/Source/Editor/SkeletonEditor/Private/SSkeletonTree.cpp
+++ b/Engine/Source/Editor/SkeletonEditor/Private/SSkeletonTree.cpp
@@ -343,6 +343,12 @@ void SSkeletonTree::BindCommands()
 		FCanExecuteAction(),
 		FIsActionChecked::CreateLambda([]() { return GetDefault<UPersonaOptions>()->bHideParentsWhenFiltering; }));
 
+	CommandList.MapAction(
+		MenuActions.ShowBoneIndex,
+		FExecuteAction::CreateLambda([this]() { GetMutableDefault<UPersonaOptions>()->bShowBoneIndexes = !GetDefault<UPersonaOptions>()->bShowBoneIndexes; }),
+		FCanExecuteAction(),
+		FIsActionChecked::CreateLambda([]() { return GetDefault<UPersonaOptions>()->bShowBoneIndexes; }));
+	
 	// Bone Filter commands
 	CommandList.BeginGroup(TEXT("BoneFilterGroup"));
 
@@ -1969,6 +1975,7 @@ void SSkeletonTree::RegisterFilterMenu()
 		OptionsSection.AddMenuEntry(Actions.FilteringFlattensHierarchy);
 		OptionsSection.AddMenuEntry(Actions.HideParentsWhenFiltering);
 		OptionsSection.AddMenuEntry(Actions.ShowDebugVisualization);
+		OptionsSection.AddMenuEntry(Actions.ShowBoneIndex);
 	}
 
 	{
diff --git a/Engine/Source/Editor/SkeletonEditor/Private/SkeletonTreeBoneItem.cpp b/Engine/Source/Editor/SkeletonEditor/Private/SkeletonTreeBoneItem.cpp
--- a/Engine/Source/Editor/SkeletonEditor/Private/SkeletonTreeBoneItem.cpp
+++ b/Engine/Source/Editor/SkeletonEditor/Private/SkeletonTreeBoneItem.cpp
@@ -23,6 +23,7 @@
 #include "UObject/Package.h"
 #include "Editor.h"
 #include "ScopedTransaction.h"
+#include "Preferences/PersonaOptions.h"
 
 #define LOCTEXT_NAMESPACE "FSkeletonTreeBoneItem"
 
@@ -82,7 +83,7 @@ void FSkeletonTreeBoneItem::GenerateWidgetForNameColumn( TSharedPtr< SHorizontal
 		UDebugSkelMeshComponent* PreviewComponent = GetSkeletonTree()->GetPreviewScene()->GetPreviewMeshComponent();
 		CacheLODChange(PreviewComponent);
 	}	
-	
+
 	FText ToolTip = GetBoneToolTip();
 	Box->AddSlot()
 		.AutoWidth()
@@ -96,6 +97,26 @@ void FSkeletonTreeBoneItem::GenerateWidgetForNameColumn( TSharedPtr< SHorizontal
 			.Font(this, &FSkeletonTreeBoneItem::GetBoneTextFont)
 			.ToolTipText( ToolTip )
 		];
+
+	Box->AddSlot()
+		.AutoWidth()
+		.Padding(4, 0, 0, 0)
+		.VAlign(VAlign_Center)
+		[
+			SNew (STextBlock)
+			.ColorAndOpacity(this, &FSkeletonTreeBoneItem::GetBoneTextColor, InIsSelected)
+			.Text_Lambda( [BoneName=BoneName, EditableSkeleton=GetEditableSkeleton()]() -> FText
+			{
+				const int32 BoneIndex = EditableSkeleton->GetSkeleton().GetReferenceSkeleton().FindBoneIndex(BoneName);
+				return FText::Format(LOCTEXT("BoneIndexValue", "<{0}>"), FText::AsNumber(BoneIndex));
+			})
+			.Font(this, &FSkeletonTreeBoneItem::GetBoneTextFont)
+			.ToolTipText( LOCTEXT("BoneIndexTooltip", "The index of this bone in the reference skeleton's flat list of bones.") )
+			.Visibility_Lambda([]() -> EVisibility
+			{
+				return GetDefault<UPersonaOptions>()->bShowBoneIndexes ? EVisibility::Visible : EVisibility::Collapsed;
+			})
+		];
 }
 
 TSharedRef< SWidget > FSkeletonTreeBoneItem::GenerateWidgetForDataColumn(const FName& DataColumnName, FIsSelected InIsSelected)
diff --git a/Engine/Source/Editor/SkeletonEditor/Private/SkeletonTreeCommands.cpp b/Engine/Source/Editor/SkeletonEditor/Private/SkeletonTreeCommands.cpp
--- a/Engine/Source/Editor/SkeletonEditor/Private/SkeletonTreeCommands.cpp
+++ b/Engine/Source/Editor/SkeletonEditor/Private/SkeletonTreeCommands.cpp
@@ -20,6 +20,7 @@ void FSkeletonTreeCommands::RegisterCommands()
 	UI_COMMAND( ShowDebugVisualization, "Show Debug Visualization Options", "Show debug visualization options for objects in the viewport", EUserInterfaceActionType::ToggleButton, FInputChord());
 	UI_COMMAND( FilteringFlattensHierarchy, "Filtering Flattens Hierarchy", "Whether to keep the hierarchy or flatten it when searching for tree items", EUserInterfaceActionType::ToggleButton, FInputChord());
 	UI_COMMAND( HideParentsWhenFiltering, "Hide Parents When Filtering", "Whether to show parent items grayed out, or hide them entirely when filtering", EUserInterfaceActionType::ToggleButton, FInputChord());
+	UI_COMMAND( ShowBoneIndex, "Show Bone Index", "Show the index of the bone as stored in the reference skeleton", EUserInterfaceActionType::ToggleButton, FInputChord());
 	
 	UI_COMMAND( CopyBoneNames, "Copy Selected Bone Names", "Copy selected bone names to clipboard", EUserInterfaceActionType::Button, FInputChord() );
 	UI_COMMAND( ResetBoneTransforms, "Reset Selected Bone Transforms", "Reset the transforms of the selected bones", EUserInterfaceActionType::Button, FInputChord() );
diff --git a/Engine/Source/Editor/SkeletonEditor/Private/SkeletonTreeCommands.h b/Engine/Source/Editor/SkeletonEditor/Private/SkeletonTreeCommands.h
--- a/Engine/Source/Editor/SkeletonEditor/Private/SkeletonTreeCommands.h
+++ b/Engine/Source/Editor/SkeletonEditor/Private/SkeletonTreeCommands.h
@@ -60,6 +60,9 @@ public:
 	/** Hide parents on filter */
 	TSharedPtr< FUICommandInfo > HideParentsWhenFiltering;
 
+	/** Show bone numbers */
+	TSharedPtr< FUICommandInfo > ShowBoneIndex;
+
 	/** Add a socket to the skeleton */
 	TSharedPtr< FUICommandInfo > AddSocket;
 
diff --git a/Engine/Source/Editor/SkeletonEditor/Private/SkeletonTreeVirtualBoneItem.cpp b/Engine/Source/Editor/SkeletonEditor/Private/SkeletonTreeVirtualBoneItem.cpp
--- a/Engine/Source/Editor/SkeletonEditor/Private/SkeletonTreeVirtualBoneItem.cpp
+++ b/Engine/Source/Editor/SkeletonEditor/Private/SkeletonTreeVirtualBoneItem.cpp
@@ -16,6 +16,7 @@
 #include "SocketDragDropOp.h"
 #include "Editor.h"
 #include "ScopedTransaction.h"
+#include "Preferences/PersonaOptions.h"
 
 #define LOCTEXT_NAMESPACE "FSkeletonTreeVirtualBoneItem"
 
@@ -92,6 +93,26 @@ void FSkeletonTreeVirtualBoneItem::GenerateWidgetForNameColumn(TSharedPtr< SHori
 		[
 			InlineWidget.ToSharedRef()
 		];
+	
+	Box->AddSlot()
+		.AutoWidth()
+		.Padding(4, 0, 0, 0)
+		.VAlign(VAlign_Center)
+		[
+			SNew (STextBlock)
+			.ColorAndOpacity(this, &FSkeletonTreeVirtualBoneItem::GetBoneTextColor, InIsSelected)
+			.Text_Lambda( [BoneName=BoneName, EditableSkeleton=GetEditableSkeleton()]() -> FText
+			{
+				const int32 BoneIndex = EditableSkeleton->GetSkeleton().GetReferenceSkeleton().FindBoneIndex(BoneName);
+				return FText::Format(LOCTEXT("VirtualBoneIndexValue", "<{0}>"), FText::AsNumber(BoneIndex));
+			})
+			.Font(this, &FSkeletonTreeVirtualBoneItem::GetBoneTextFont)
+			.ToolTipText( LOCTEXT("VirtualBoneIndexTooltip", "The index of this virtual bone in the reference skeleton's flat list of bones.") )
+			.Visibility_Lambda([]() -> EVisibility
+			{
+				return GetDefault<UPersonaOptions>()->bShowBoneIndexes ? EVisibility::Visible : EVisibility::Collapsed;
+			})
+		];
 }
 
 TSharedRef< SWidget > FSkeletonTreeVirtualBoneItem::GenerateWidgetForDataColumn(const FName& DataColumnName, FIsSelected InIsSelected)
diff --git a/Engine/Source/Editor/UnrealEd/Classes/Preferences/PersonaOptions.h b/Engine/Source/Editor/UnrealEd/Classes/Preferences/PersonaOptions.h
--- a/Engine/Source/Editor/UnrealEd/Classes/Preferences/PersonaOptions.h
+++ b/Engine/Source/Editor/UnrealEd/Classes/Preferences/PersonaOptions.h
@@ -186,6 +186,10 @@ class UPersonaOptions : public UObject
 	UPROPERTY(EditAnywhere, config, Category = "Skeleton Tree")
 	bool bHideParentsWhenFiltering;
 
+	/** Whether to append each bone's index to its name when shown in the skeleton tree */
+	UPROPERTY(EditAnywhere, config, Category = "Skeleton Tree")
+	bool bShowBoneIndexes = false;
+
 	/** Whether to focus and expand an item's tree recursively based on selection */
 	UPROPERTY(EditAnywhere, config, Category = "Skeleton Tree")
 	bool bExpandTreeOnSelection;
