Remove "selected" from bone context menu

diff --git a/Engine/Source/Editor/SkeletonEditor/Private/SkeletonTreeCommands.cpp b/Engine/Source/Editor/SkeletonEditor/Private/SkeletonTreeCommands.cpp
--- a/Engine/Source/Editor/SkeletonEditor/Private/SkeletonTreeCommands.cpp
+++ b/Engine/Source/Editor/SkeletonEditor/Private/SkeletonTreeCommands.cpp
@@ -22,8 +22,8 @@ void FSkeletonTreeCommands::RegisterCommands()
 	UI_COMMAND( HideParentsWhenFiltering, "Hide Parents When Filtering", "Whether to show parent items grayed out, or hide them entirely when filtering", EUserInterfaceActionType::ToggleButton, FInputChord());
 	UI_COMMAND( ShowBoneIndex, "Show Bone Index", "Show the index of the bone as stored in the reference skeleton", EUserInterfaceActionType::ToggleButton, FInputChord());
 	
-	UI_COMMAND( CopyBoneNames, "Copy Selected Bone Names", "Copy selected bone names to clipboard", EUserInterfaceActionType::Button, FInputChord() );
-	UI_COMMAND( ResetBoneTransforms, "Reset Selected Bone Transforms", "Reset the transforms of the selected bones", EUserInterfaceActionType::Button, FInputChord() );
+	UI_COMMAND( CopyBoneNames, "Copy Bone Names", "Copy selected bone names to clipboard", EUserInterfaceActionType::Button, FInputChord() );
+	UI_COMMAND( ResetBoneTransforms, "Reset Bone Transforms", "Reset the transforms of the selected bones", EUserInterfaceActionType::Button, FInputChord() );
 	UI_COMMAND( CopySockets, "Copy Selected Sockets", "Copy selected sockets to clipboard", EUserInterfaceActionType::Button, FInputChord( EModifierKey::Control, EKeys::C ) );
 	UI_COMMAND( PasteSockets, "Paste Sockets", "Paste sockets from clipboard (try to add them to the original copy bone or root if that bone is not available)", EUserInterfaceActionType::Button, FInputChord( EModifierKey::Control, EKeys::V ) );
 	UI_COMMAND( PasteSocketsToSelectedBone, "Paste Sockets To Selected Bone", "Paste sockets from clipboard to selected bone", EUserInterfaceActionType::Button, FInputChord(EModifierKey::Control | EModifierKey::Shift, EKeys::V));
