[UE] PhAT Fix range selection inconsistancy between viewport and skeleton tree
- SListView keeps track of a RangeStart Item which is used to determine the range of items selected in the tree view. This is set to the current selection only when select info indicates a user interaction (click / button press).
- Changed SSkeletonTree to pass a SelectIndo type trhough to the SListView on range selection.
- Changed PhysicsAssetEditor to flag selection in the viewport as a user interaction

diff --git a/Engine/Source/Editor/PhysicsAssetEditor/Private/PhysicsAssetEditor.cpp b/Engine/Source/Editor/PhysicsAssetEditor/Private/PhysicsAssetEditor.cpp
--- a/Engine/Source/Editor/PhysicsAssetEditor/Private/PhysicsAssetEditor.cpp
+++ b/Engine/Source/Editor/PhysicsAssetEditor/Private/PhysicsAssetEditor.cpp
@@ -393,7 +393,8 @@ void FPhysicsAssetEditor::HandleViewportSelectionChanged(const TArray<FPhysicsAs
 						}
 					}
 					return false;
-				});
+				},
+				ESelectInfo::OnMouseClick);
 			}
 
 			if (PhysicsAssetGraph.IsValid())
diff --git a/Engine/Source/Editor/SkeletonEditor/Private/SSkeletonTree.cpp b/Engine/Source/Editor/SkeletonEditor/Private/SSkeletonTree.cpp
--- a/Engine/Source/Editor/SkeletonEditor/Private/SSkeletonTree.cpp
+++ b/Engine/Source/Editor/SkeletonEditor/Private/SSkeletonTree.cpp
@@ -2368,7 +2368,7 @@ void SSkeletonTree::OnPreviewMeshChanged(USkeletalMesh* InOldSkeletalMesh, USkel
 	}
 }
 
-void SSkeletonTree::SelectItemsBy(TFunctionRef<bool(const TSharedRef<ISkeletonTreeItem>&, bool&)> Predicate) const
+void SSkeletonTree::SelectItemsBy(TFunctionRef<bool(const TSharedRef<ISkeletonTreeItem>&, bool&)> Predicate, const ESelectInfo::Type SelectInfo) const
 {
 	TArray<TPair<TSharedPtr<ISkeletonTreeItem>, bool>> ItemsToSelect;
 	TSharedPtr<ISkeletonTreeItem> ScrollItem = nullptr;
@@ -2393,7 +2393,7 @@ void SSkeletonTree::SelectItemsBy(TFunctionRef<bool(const TSharedRef<ISkeletonTr
 		{
 			TSharedPtr<ISkeletonTreeItem> Item = ItemPair.Key;
 
-			SkeletonTreeView->SetItemSelection(Item, true);
+			SkeletonTreeView->SetItemSelection(Item, true, SelectInfo);
 			if (ItemPair.Value)
 			{
 				if(Item->GetChildren().Num() == 0)
diff --git a/Engine/Source/Editor/SkeletonEditor/Private/SSkeletonTree.h b/Engine/Source/Editor/SkeletonEditor/Private/SSkeletonTree.h
--- a/Engine/Source/Editor/SkeletonEditor/Private/SSkeletonTree.h
+++ b/Engine/Source/Editor/SkeletonEditor/Private/SSkeletonTree.h
@@ -79,7 +79,7 @@ public:
 	virtual void SetSelectedBone(const FName& InBoneName, ESelectInfo::Type InSelectInfo) override;
 	virtual void DeselectAll() override;
 	virtual TArray<TSharedPtr<ISkeletonTreeItem>> GetSelectedItems() const override { return SkeletonTreeView->GetSelectedItems(); }
-	virtual void SelectItemsBy(TFunctionRef<bool(const TSharedRef<ISkeletonTreeItem>&, bool&)> Predicate) const override;
+	virtual void SelectItemsBy(TFunctionRef<bool(const TSharedRef<ISkeletonTreeItem>&, bool&)> Predicate, const ESelectInfo::Type SelectInfo) const override;
 	virtual void DuplicateAndSelectSocket(const FSelectedSocketInfo& SocketInfoToDuplicate, const FName& NewParentBoneName = FName()) override;
 
 	virtual FDelegateHandle RegisterOnSelectionChanged(const FOnSkeletonTreeSelectionChanged& Delegate) override
diff --git a/Engine/Source/Editor/SkeletonEditor/Public/ISkeletonTree.h b/Engine/Source/Editor/SkeletonEditor/Public/ISkeletonTree.h
--- a/Engine/Source/Editor/SkeletonEditor/Public/ISkeletonTree.h
+++ b/Engine/Source/Editor/SkeletonEditor/Public/ISkeletonTree.h
@@ -140,7 +140,7 @@ public:
 	virtual TArray<TSharedPtr<ISkeletonTreeItem>> GetSelectedItems() const = 0;
 
 	/** Select items using the passed in predicate  */
-	virtual void SelectItemsBy(TFunctionRef<bool(const TSharedRef<ISkeletonTreeItem>& /* InItem */, bool& /* bInOutExpand */)> Predicate) const = 0;
+	virtual void SelectItemsBy(TFunctionRef<bool(const TSharedRef<ISkeletonTreeItem>& /* InItem */, bool& /* bInOutExpand */)> Predicate, const ESelectInfo::Type SelectInfo = ESelectInfo::Direct) const = 0;
 
 	/** Duplicate the socket and select it */
 	virtual void DuplicateAndSelectSocket(const FSelectedSocketInfo& SocketInfoToDuplicate, const FName& NewParentBoneName = FName()) = 0;
