Skeletal Mesh Editor: quick way to reset selected / all bones transforms

diff --git a/Engine/Source/Editor/SkeletalMeshEditor/Private/SkeletalMeshEditor.cpp b/Engine/Source/Editor/SkeletalMeshEditor/Private/SkeletalMeshEditor.cpp
--- a/Engine/Source/Editor/SkeletalMeshEditor/Private/SkeletalMeshEditor.cpp
+++ b/Engine/Source/Editor/SkeletalMeshEditor/Private/SkeletalMeshEditor.cpp
@@ -310,22 +310,31 @@ void FSkeletalMeshEditor::BindCommands()
 	constexpr bool bWithDialogFalse = false;
 	constexpr bool bWithDialogTrue = true;
 
-	ToolkitCommands->MapAction(FSkeletalMeshEditorCommands::Get().ReimportMesh,
+	const FSkeletalMeshEditorCommands Commands = FSkeletalMeshEditorCommands::Get();
+	ToolkitCommands->MapAction(Commands.ReimportMesh,
 		FExecuteAction::CreateSP(this, &FSkeletalMeshEditor::HandleReimportMesh, FReimportParameters((int32)INDEX_NONE, bWithNewFileFalse, bWithDialogFalse)));
 
-	ToolkitCommands->MapAction(FSkeletalMeshEditorCommands::Get().ReimportAllMesh,
+	ToolkitCommands->MapAction(Commands.ReimportAllMesh,
 		FExecuteAction::CreateSP(this, &FSkeletalMeshEditor::HandleReimportAllMesh, FReimportParameters((int32)INDEX_NONE, bWithNewFileFalse, bWithDialogFalse)));
 
-	ToolkitCommands->MapAction(FSkeletalMeshEditorCommands::Get().ReimportWithDialog,
+	ToolkitCommands->MapAction(Commands.ReimportWithDialog,
 		FExecuteAction::CreateSP(this, &FSkeletalMeshEditor::HandleReimportMesh, FReimportParameters((int32)INDEX_NONE, bWithNewFileFalse, bWithDialogTrue)));
 
 	ToolkitCommands->MapAction(FPersonaCommonCommands::Get().TogglePlay,
 		FExecuteAction::CreateRaw(&GetPersonaToolkit()->GetPreviewScene().Get(), &IPersonaPreviewScene::TogglePlayback));
 
 	// Bake Materials
-	ToolkitCommands->MapAction(FSkeletalMeshEditorCommands::Get().BakeMaterials,
+	ToolkitCommands->MapAction(Commands.BakeMaterials,
 		FExecuteAction::CreateSP(this, &FSkeletalMeshEditor::BakeMaterials),
 		FCanExecuteAction());
+
+	// Reset transforms
+	static constexpr bool bSelectedOnly = true;
+	ToolkitCommands->MapAction(Commands.ResetBoneTransforms,
+		FExecuteAction::CreateSP(GetSkeletonTree(), &ISkeletonTree::ResetBoneTransforms, bSelectedOnly));
+	
+	ToolkitCommands->MapAction(Commands.ResetAllBonesTransforms,
+		FExecuteAction::CreateSP(GetSkeletonTree(), &ISkeletonTree::ResetBoneTransforms, !bSelectedOnly));
 }
 
 TSharedPtr<FSkeletalMeshEditor> FSkeletalMeshEditor::GetSkeletalMeshEditor(const FToolMenuContext& InMenuContext)
diff --git a/Engine/Source/Editor/SkeletalMeshEditor/Private/SkeletalMeshEditorCommands.cpp b/Engine/Source/Editor/SkeletalMeshEditor/Private/SkeletalMeshEditorCommands.cpp
--- a/Engine/Source/Editor/SkeletalMeshEditor/Private/SkeletalMeshEditorCommands.cpp
+++ b/Engine/Source/Editor/SkeletalMeshEditor/Private/SkeletalMeshEditorCommands.cpp
@@ -18,6 +18,9 @@ void FSkeletalMeshEditorCommands::RegisterCommands()
 	UI_COMMAND(ReimportWithDialog, "Reimport Mesh With Dialog", "Reimport the base mesh with dialog.", EUserInterfaceActionType::Button, FInputChord());
 
 	UI_COMMAND(BakeMaterials, "Bake Out Materials", "Bake out Materials for given LOD(s).", EUserInterfaceActionType::Button, FInputChord());
+
+	UI_COMMAND(ResetBoneTransforms, "Reset Bone Transforms", "Reset selected bones transforms.", EUserInterfaceActionType::Button, FInputChord(EKeys::I, EModifierKey::Control) );
+	UI_COMMAND(ResetAllBonesTransforms, "Reset All Bones Transforms", "Reset all bones transforms.", EUserInterfaceActionType::Button, FInputChord(EKeys::I, EModifierKey::Control | EModifierKey::Shift) );
 }
 
 #undef LOCTEXT_NAMESPACE
diff --git a/Engine/Source/Editor/SkeletalMeshEditor/Private/SkeletalMeshEditorCommands.h b/Engine/Source/Editor/SkeletalMeshEditor/Private/SkeletalMeshEditorCommands.h
--- a/Engine/Source/Editor/SkeletalMeshEditor/Private/SkeletalMeshEditorCommands.h
+++ b/Engine/Source/Editor/SkeletalMeshEditor/Private/SkeletalMeshEditorCommands.h
@@ -34,4 +34,8 @@ public:
 	
 	// bake materials for this skeletal mesh
 	TSharedPtr<FUICommandInfo> BakeMaterials;
+
+	// Reset bones transforms
+	TSharedPtr< FUICommandInfo > ResetBoneTransforms;
+	TSharedPtr< FUICommandInfo > ResetAllBonesTransforms;
 };
diff --git a/Engine/Source/Editor/SkeletonEditor/Private/SSkeletonTree.cpp b/Engine/Source/Editor/SkeletonEditor/Private/SSkeletonTree.cpp
--- a/Engine/Source/Editor/SkeletonEditor/Private/SSkeletonTree.cpp
+++ b/Engine/Source/Editor/SkeletonEditor/Private/SSkeletonTree.cpp
@@ -476,9 +476,14 @@ void SSkeletonTree::BindCommands()
 		MenuActions.CopyBoneNames,
 		FExecuteAction::CreateSP( this, &SSkeletonTree::OnCopyBoneNames ));
 
+	static constexpr bool bSelectedOnly = true;
 	CommandList.MapAction(
 		MenuActions.ResetBoneTransforms,
-		FExecuteAction::CreateSP(this, &SSkeletonTree::OnResetBoneTransforms ) );
+		FExecuteAction::CreateSP(this, &SSkeletonTree::OnResetBoneTransforms, bSelectedOnly ) );
+
+	CommandList.MapAction(
+		MenuActions.ResetAllBonesTransforms,
+		FExecuteAction::CreateSP(this, &SSkeletonTree::OnResetBoneTransforms, !bSelectedOnly ) );
 
 	CommandList.MapAction(
 		MenuActions.CopySockets,
@@ -1200,7 +1205,7 @@ void SSkeletonTree::OnCopyBoneNames()
 	}
 }
 
-void SSkeletonTree::OnResetBoneTransforms()
+void SSkeletonTree::OnResetBoneTransforms(const bool bSelectedOnly)
 {
 	if (GetPreviewScene().IsValid())
 	{
@@ -1209,6 +1214,24 @@ void SSkeletonTree::OnResetBoneTransforms()
 		UAnimPreviewInstance* PreviewInstance = PreviewComponent->PreviewInstance;
 		check(PreviewInstance);
 
+		const TArray<FAnimNode_ModifyBone>& ModifyBones = PreviewInstance->GetBoneControllers();
+		if (ModifyBones.IsEmpty())
+		{
+			return;
+		}
+		
+		if (!bSelectedOnly)
+		{
+			GEditor->BeginTransaction(LOCTEXT("SkeletonTree_ResetAllBonesTransforms", "Reset All Bone Transforms"));
+			{
+				PreviewInstance->SetFlags(RF_Transactional);
+				PreviewInstance->Modify();
+				PreviewInstance->ResetModifiedBone();
+			}
+			GEditor->EndTransaction();
+			return;
+		}
+		
 		TArray<TSharedPtr<ISkeletonTreeItem>> SelectedItems = SkeletonTreeView->GetSelectedItems();
 		FSkeletonTreeSelection TreeSelection(SelectedItems);
 		TArray<TSharedPtr<FSkeletonTreeBoneItem>> SelectedBones = TreeSelection.GetSelectedItems<FSkeletonTreeBoneItem>();
@@ -2435,6 +2458,11 @@ void SSkeletonTree::DuplicateAndSelectSocket(const FSelectedSocketInfo& SocketIn
 	SetSelectedSocket(NewSocketInfo);
 }
 
+void SSkeletonTree::ResetBoneTransforms(const bool bSelectedOnly)
+{
+	OnResetBoneTransforms(bSelectedOnly);
+}
+
 void SSkeletonTree::HandleFocusCamera()
 {
 	if (GetPreviewScene().IsValid())
diff --git a/Engine/Source/Editor/SkeletonEditor/Private/SSkeletonTree.h b/Engine/Source/Editor/SkeletonEditor/Private/SSkeletonTree.h
--- a/Engine/Source/Editor/SkeletonEditor/Private/SSkeletonTree.h
+++ b/Engine/Source/Editor/SkeletonEditor/Private/SSkeletonTree.h
@@ -81,6 +81,7 @@ public:
 	virtual TArray<TSharedPtr<ISkeletonTreeItem>> GetSelectedItems() const override { return SkeletonTreeView->GetSelectedItems(); }
 	virtual void SelectItemsBy(TFunctionRef<bool(const TSharedRef<ISkeletonTreeItem>&, bool&)> Predicate, const ESelectInfo::Type SelectInfo) const override;
 	virtual void DuplicateAndSelectSocket(const FSelectedSocketInfo& SocketInfoToDuplicate, const FName& NewParentBoneName = FName()) override;
+	virtual void ResetBoneTransforms(const bool bSelectedOnly) override;
 
 	virtual FDelegateHandle RegisterOnSelectionChanged(const FOnSkeletonTreeSelectionChanged& Delegate) override
 	{
@@ -184,7 +185,7 @@ private:
 	void OnCopyBoneNames();
 
 	/** Function to reset the transforms of selected bones */
-	void OnResetBoneTransforms();
+	void OnResetBoneTransforms(const bool bSelectedOnly = true);
 
 	/** Function to copy selected sockets to the clipboard */
 	void OnCopySockets() const;
diff --git a/Engine/Source/Editor/SkeletonEditor/Private/SkeletonTreeCommands.cpp b/Engine/Source/Editor/SkeletonEditor/Private/SkeletonTreeCommands.cpp
--- a/Engine/Source/Editor/SkeletonEditor/Private/SkeletonTreeCommands.cpp
+++ b/Engine/Source/Editor/SkeletonEditor/Private/SkeletonTreeCommands.cpp
@@ -23,7 +23,8 @@ void FSkeletonTreeCommands::RegisterCommands()
 	UI_COMMAND( ShowBoneIndex, "Show Bone Index", "Show the index of the bone as stored in the reference skeleton", EUserInterfaceActionType::ToggleButton, FInputChord());
 	
 	UI_COMMAND( CopyBoneNames, "Copy Bone Names", "Copy selected bone names to clipboard", EUserInterfaceActionType::Button, FInputChord() );
-	UI_COMMAND( ResetBoneTransforms, "Reset Bone Transforms", "Reset the transforms of the selected bones", EUserInterfaceActionType::Button, FInputChord() );
+	UI_COMMAND( ResetBoneTransforms, "Reset Bone Transforms", "Reset the transforms of the selected bones", EUserInterfaceActionType::Button, FInputChord(EKeys::I, EModifierKey::Control) );
+	UI_COMMAND( ResetAllBonesTransforms, "Reset All Bones Transforms", "Reset the transforms of all the bones", EUserInterfaceActionType::Button, FInputChord(EKeys::I, EModifierKey::Control | EModifierKey::Shift) );
 	UI_COMMAND( CopySockets, "Copy Selected Sockets", "Copy selected sockets to clipboard", EUserInterfaceActionType::Button, FInputChord( EModifierKey::Control, EKeys::C ) );
 	UI_COMMAND( PasteSockets, "Paste Sockets", "Paste sockets from clipboard (try to add them to the original copy bone or root if that bone is not available)", EUserInterfaceActionType::Button, FInputChord( EModifierKey::Control, EKeys::V ) );
 	UI_COMMAND( PasteSocketsToSelectedBone, "Paste Sockets To Selected Bone", "Paste sockets from clipboard to selected bone", EUserInterfaceActionType::Button, FInputChord(EModifierKey::Control | EModifierKey::Shift, EKeys::V));
diff --git a/Engine/Source/Editor/SkeletonEditor/Private/SkeletonTreeCommands.h b/Engine/Source/Editor/SkeletonEditor/Private/SkeletonTreeCommands.h
--- a/Engine/Source/Editor/SkeletonEditor/Private/SkeletonTreeCommands.h
+++ b/Engine/Source/Editor/SkeletonEditor/Private/SkeletonTreeCommands.h
@@ -96,9 +96,12 @@ public:
 	/** Copy bone names */
 	TSharedPtr< FUICommandInfo > CopyBoneNames;
 
-	/** Reset bone transforms */
+	/** Reset selected bones transforms */
 	TSharedPtr< FUICommandInfo > ResetBoneTransforms;
 
+	/** Reset all bones transforms */
+	TSharedPtr< FUICommandInfo > ResetAllBonesTransforms;
+
 	/** Copy sockets to clipboard */
 	TSharedPtr< FUICommandInfo > CopySockets;
 
diff --git a/Engine/Source/Editor/SkeletonEditor/Public/ISkeletonTree.h b/Engine/Source/Editor/SkeletonEditor/Public/ISkeletonTree.h
--- a/Engine/Source/Editor/SkeletonEditor/Public/ISkeletonTree.h
+++ b/Engine/Source/Editor/SkeletonEditor/Public/ISkeletonTree.h
@@ -162,4 +162,7 @@ public:
 
 	/** Get the pinned commands widget, if any, for this tree */
 	virtual TSharedPtr<IPinnedCommandList> GetPinnedCommandList() const = 0;
+
+	/** Function to reset the transforms of selected bones */
+	virtual void ResetBoneTransforms(const bool bSelectedOnly) = 0;
 };
