Skeletal Mesh Editor: multiple bone selection
- that this is only enabled using the ITF version of the drag tool (SelectionMode.EnableITFTools 1)
- note this version is faster than sequential bone selection

TODO
- enable transforming multiple joints at a time (UE-184844)
- support shift + click to select multiple joints (UE-170097)

diff --git a/Engine/Source/Editor/Persona/Private/AnimationEditorPreviewScene.cpp b/Engine/Source/Editor/Persona/Private/AnimationEditorPreviewScene.cpp
--- a/Engine/Source/Editor/Persona/Private/AnimationEditorPreviewScene.cpp
+++ b/Engine/Source/Editor/Persona/Private/AnimationEditorPreviewScene.cpp
@@ -849,35 +849,72 @@ void FAnimationEditorPreviewScene::ClearSelectedBone()
 		SkeletalMeshComponent->BonesOfInterest.Empty();
 	}
 
-	OnSelectedBoneChanged.Broadcast(NAME_None, ESelectInfo::Direct);
+	PRAGMA_DISABLE_DEPRECATION_WARNINGS
+	OnSelectedBoneChanged_DEPRECATED.Broadcast({NAME_None}, ESelectInfo::Direct);
+	PRAGMA_ENABLE_DEPRECATION_WARNINGS
+	
+	OnSelectedBonesChanged.Broadcast({NAME_None}, ESelectInfo::Direct);
 
 	InvalidateViews();
 }
 
 void FAnimationEditorPreviewScene::SetSelectedBone(const FName& BoneName, ESelectInfo::Type InSelectInfo)
 {
-	TGuardValue<bool> RecursionGuard(bSelecting, true);
-
 	const int32 BoneIndex = GetEditableSkeleton().IsValid() ? GetEditableSkeleton()->GetSkeleton().GetReferenceSkeleton().FindBoneIndex(BoneName) : INDEX_NONE;
 	if (BoneIndex != INDEX_NONE)
 	{
-		ClearSelectedBone();
-		ClearSelectedSocket();
-		ClearSelectedActor();
+		SetSelectedBones({BoneName}, InSelectInfo);
+
+		PRAGMA_DISABLE_DEPRECATION_WARNINGS
+        OnSelectedBoneChanged_DEPRECATED.Broadcast(BoneName, InSelectInfo);
+        PRAGMA_ENABLE_DEPRECATION_WARNINGS
+	}
+}
 
-		// need to get mesh bone base since BonesOfInterest is saved in SkeletalMeshComponent
-		// and it is used by renderer. It is not Skeleton base
+void FAnimationEditorPreviewScene::SetSelectedBones(const TArray<FName>& InBoneNames, ESelectInfo::Type InSelectInfo)
+{
+	if (!GetEditableSkeleton().IsValid())
+	{
+		return;
+	}
+
+	const FReferenceSkeleton& ReferenceSkeleton = GetEditableSkeleton()->GetSkeleton().GetReferenceSkeleton();
+	
+	TArray<FName> BoneNames;
+	Algo::CopyIf(InBoneNames, BoneNames, [&ReferenceSkeleton](const FName& BoneName)
+	{
+		return ReferenceSkeleton.FindBoneIndex(BoneName) != INDEX_NONE;
+	});
+
+	if (BoneNames.IsEmpty())
+	{
+		return;
+	}
+	
+	TGuardValue<bool> RecursionGuard(bSelecting, true);
+	
+	ClearSelectedBone();
+	ClearSelectedSocket();
+	ClearSelectedActor();
+	
+	// need to get mesh bone base since BonesOfInterest is saved in SkeletalMeshComponent
+	// and it is used by renderer. It is not Skeleton base
+	for (const FName& BoneName : BoneNames)
+	{
 		const int32 MeshBoneIndex = SkeletalMeshComponent->GetBoneIndex(BoneName);
 		if (MeshBoneIndex != INDEX_NONE)
 		{
-			SelectedBoneIndex = MeshBoneIndex;
+			if (SelectedBoneIndex == INDEX_NONE)
+			{
+				SelectedBoneIndex = MeshBoneIndex;
+			}
 			SkeletalMeshComponent->BonesOfInterest.Add(SelectedBoneIndex);
 		}
-		
-		InvalidateViews();
-
-		OnSelectedBoneChanged.Broadcast(BoneName, InSelectInfo);
 	}
+
+	InvalidateViews();
+	
+	OnSelectedBonesChanged.Broadcast(BoneNames, InSelectInfo);
 }
 
 void FAnimationEditorPreviewScene::SetSelectedSocket(const FSelectedSocketInfo& SocketInfo)
diff --git a/Engine/Source/Editor/Persona/Private/AnimationEditorPreviewScene.h b/Engine/Source/Editor/Persona/Private/AnimationEditorPreviewScene.h
--- a/Engine/Source/Editor/Persona/Private/AnimationEditorPreviewScene.h
+++ b/Engine/Source/Editor/Persona/Private/AnimationEditorPreviewScene.h
@@ -46,6 +46,7 @@ public:
 	virtual void ShowReferencePose(bool bShowRefPose, bool bResetBoneTransforms = false) override;
 	virtual bool IsShowReferencePoseEnabled() const override;
 	virtual void SetSelectedBone(const FName& BoneName, ESelectInfo::Type InSelectInfo) override;
+	virtual void SetSelectedBones(const TArray<FName>& InBoneNames, ESelectInfo::Type InSelectInfo) override;
 	virtual void ClearSelectedBone() override;
 	virtual void SetSelectedSocket(const FSelectedSocketInfo& SocketInfo) override;
 	virtual void ClearSelectedSocket() override;
@@ -139,14 +140,26 @@ public:
 		OnMeshClick.RemoveAll(Thing);
 	}
 
+	UE_DEPRECATED(5.6, "Use RegisterOnSelectedBonesChanged delegate instead.")	
 	virtual FDelegateHandle RegisterOnSelectedBoneChanged(const FOnSelectedBoneChanged& Delegate) override
 	{
-		return OnSelectedBoneChanged.Add(Delegate);
+		return OnSelectedBoneChanged_DEPRECATED.Add(Delegate);
 	}
-
+	
+	UE_DEPRECATED(5.6, "Use UnregisterOnSelectedBonesChanged delegate instead.")
 	virtual void UnregisterOnSelectedBoneChanged(FDelegateHandle InHandle) override
 	{
-		OnSelectedBoneChanged.Remove(InHandle);
+		OnSelectedBoneChanged_DEPRECATED.Remove(InHandle);
+	}
+	
+	virtual FDelegateHandle RegisterOnSelectedBonesChanged(const FOnSelectedBonesChanged& InDelegate) override
+	{
+		return OnSelectedBonesChanged.Add(InDelegate);
+	}
+
+	virtual void UnregisterOnSelectedBonesChanged(FDelegateHandle InHandle) override
+	{
+		OnSelectedBonesChanged.Remove(InHandle);
 	}
 
 	virtual FDelegateHandle RegisterOnSelectedSocketChanged(const FOnSelectedSocketChanged& Delegate) override
@@ -510,8 +523,12 @@ private:
 
 	/** The last time we were flagged for ticking */
 	double LastTickTime;
+	
+	/** Bone selection changed delegate */
+	FOnSelectedBoneChangedMulticaster OnSelectedBoneChanged_DEPRECATED;
+
 	/** Bone selection changed delegate */
-	FOnSelectedBoneChangedMulticaster OnSelectedBoneChanged;
+	FOnSelectedBonesChangedMulticaster OnSelectedBonesChanged;
 
 	/** Socket selection changed delegate */
 	FOnSelectedSocketChangedMulticaster OnSelectedSocketChanged;
diff --git a/Engine/Source/Editor/Persona/Private/AnimationEditorViewportClient.cpp b/Engine/Source/Editor/Persona/Private/AnimationEditorViewportClient.cpp
--- a/Engine/Source/Editor/Persona/Private/AnimationEditorViewportClient.cpp
+++ b/Engine/Source/Editor/Persona/Private/AnimationEditorViewportClient.cpp
@@ -285,7 +285,7 @@ FAnimationViewportClient::~FAnimationViewportClient()
 			}
 		}
 
-		AnimationEditorPreviewScene->UnregisterOnSelectedBoneChanged(OnSelectedBoneChangedHandle);
+		AnimationEditorPreviewScene->UnregisterOnSelectedBonesChanged(OnSelectedBonesChangedHandle);
 	}
 	OnPhysicsCreatedDelegateHandle.Reset();
 	OnMeshChangedDelegateHandle.Reset();
@@ -308,10 +308,10 @@ void FAnimationViewportClient::Initialize()
 		AnimationEditorPreviewScene->RegisterOnPreTick(FSimpleDelegate::CreateSP(this, &FAnimationViewportClient::HandlePreviewScenePreTick));
 		AnimationEditorPreviewScene->RegisterOnPostTick(FSimpleDelegate::CreateSP(this, &FAnimationViewportClient::HandlePreviewScenePostTick));
 
-		OnSelectedBoneChangedHandle = AnimationEditorPreviewScene->RegisterOnSelectedBoneChanged(FOnSelectedBoneChanged::CreateLambda([this](const FName&, ESelectInfo::Type)
-		{
-			UpdateBonesToDraw();
-		}));
+		OnSelectedBonesChangedHandle = AnimationEditorPreviewScene->RegisterOnSelectedBonesChanged(FOnSelectedBonesChanged::CreateLambda([this](const TArray<FName>&, ESelectInfo::Type)
+        {
+        	UpdateBonesToDraw();
+        }));
 	}
 
 	// Register delegate to update the show flags when the post processing is turned on or off
diff --git a/Engine/Source/Editor/Persona/Private/SkeletonSelectionEditMode.cpp b/Engine/Source/Editor/Persona/Private/SkeletonSelectionEditMode.cpp
--- a/Engine/Source/Editor/Persona/Private/SkeletonSelectionEditMode.cpp
+++ b/Engine/Source/Editor/Persona/Private/SkeletonSelectionEditMode.cpp
@@ -730,4 +730,53 @@ bool FSkeletonSelectionEditMode::CanCycleWidgetMode() const
 	return (SelectedBoneIndex >= 0 || SelectedSocket || SelectedActor != nullptr);
 }
 
+bool FSkeletonSelectionEditMode::FrustumSelect(const FConvexVolume& InFrustum, FEditorViewportClient* InViewportClient, bool InSelect)
+{
+	if (FAnimationViewportClient* AnimViewportClient = static_cast<FAnimationViewportClient*>(InViewportClient))
+	{
+		const EBoneDrawMode::Type BoneDrawMode = AnimViewportClient->GetBoneDrawMode();
+		if (BoneDrawMode != EBoneDrawMode::None)
+		{
+			const float BoneRadius = AnimViewportClient->GetBoneDrawSize();
+			TArray<FName> Selection;
+			
+			const UDebugSkelMeshComponent* PreviewMeshComponent = GetAnimPreviewScene().GetPreviewMeshComponent();
+			const FReferenceSkeleton& RefSkeleton = PreviewMeshComponent->GetReferenceSkeleton();
+			const TArray<FTransform>& BoneTransforms = PreviewMeshComponent->GetComponentSpaceTransforms();
+			for (int32 Index = 0; Index < BoneTransforms.Num(); ++Index)
+			{
+				const FVector BoneLocation = BoneTransforms[Index].GetLocation();
+				
+				// NOTE: the full sphere is used for selection for simplicity, not a "wireframe" version
+				if (InFrustum.IntersectSphere(BoneLocation, 0.1f * BoneRadius))
+				{
+					Selection.AddUnique(RefSkeleton.GetBoneName(Index));
+				}
+				else
+				{
+					// otherwise, intersect the line between the bone and its children
+					TArray<int32> Children;
+					RefSkeleton.GetDirectChildBones(Index, Children);
+					for (const int32 ChildIndex : Children)
+					{
+						if (InFrustum.IntersectLineSegment(BoneLocation, BoneTransforms[ChildIndex].GetLocation()))
+						{
+							Selection.AddUnique(RefSkeleton.GetBoneName(Index));
+							break;
+						}
+					}
+				}
+			}
+			
+			if (!Selection.IsEmpty())
+			{
+				GetAnimPreviewScene().SetSelectedBones(Selection, ESelectInfo::Direct);
+				return true;
+			}
+		}
+	}
+	
+	return IPersonaEditMode::FrustumSelect(InFrustum, InViewportClient, InSelect);
+}
+
 #undef LOCTEXT_NAMESPACE
diff --git a/Engine/Source/Editor/Persona/Public/AnimationEditorViewportClient.h b/Engine/Source/Editor/Persona/Public/AnimationEditorViewportClient.h
--- a/Engine/Source/Editor/Persona/Public/AnimationEditorViewportClient.h
+++ b/Engine/Source/Editor/Persona/Public/AnimationEditorViewportClient.h
@@ -358,7 +358,7 @@ private:
 	/* Bit field indexed on bone index that stores what bones are visible in the viewport, updated with UpdateBonesToDraw */
 	TBitArray<> BonesToDraw;
 
-	FDelegateHandle OnSelectedBoneChangedHandle;
+	FDelegateHandle OnSelectedBonesChangedHandle;
 
 private:
 
diff --git a/Engine/Source/Editor/Persona/Public/IPersonaPreviewScene.h b/Engine/Source/Editor/Persona/Public/IPersonaPreviewScene.h
--- a/Engine/Source/Editor/Persona/Public/IPersonaPreviewScene.h
+++ b/Engine/Source/Editor/Persona/Public/IPersonaPreviewScene.h
@@ -43,6 +43,10 @@ typedef FOnSelectedLODChangedMulticaster::FDelegate FOnSelectedLODChanged;
 DECLARE_MULTICAST_DELEGATE_TwoParams(FOnSelectedBoneChangedMulticaster, const FName& /*InBoneName*/, ESelectInfo::Type /*InSelectInfo*/);
 typedef FOnSelectedBoneChangedMulticaster::FDelegate FOnSelectedBoneChanged;
 
+//The bones selection changed
+DECLARE_MULTICAST_DELEGATE_TwoParams(FOnSelectedBonesChangedMulticaster, const TArray<FName>& /*InBoneNames*/, ESelectInfo::Type /*InSelectInfo*/);
+typedef FOnSelectedBonesChangedMulticaster::FDelegate FOnSelectedBonesChanged;
+
 //The selected socket changed
 DECLARE_MULTICAST_DELEGATE_OneParam(FOnSelectedSocketChangedMulticaster, const FSelectedSocketInfo& /*InSocketInfo*/);
 typedef FOnSelectedSocketChangedMulticaster::FDelegate FOnSelectedSocketChanged;
@@ -138,6 +142,9 @@ public:
 	/** Sets the selected bone on the preview component */
 	virtual void SetSelectedBone(const FName& BoneName, ESelectInfo::Type InSelectInfo) = 0;
 
+	/** Sets the selected bones on the preview component */
+	virtual void SetSelectedBones(const TArray<FName>& InBoneNames, ESelectInfo::Type InSelectInfo) = 0;
+
 	/** Clears the selected bone on the preview component */
 	virtual void ClearSelectedBone() = 0;
 
@@ -205,11 +212,19 @@ public:
 	virtual void UnregisterOnMeshClick(FDelegateUserObject Thing) = 0;
 
 	/** Registers a delegate to be called when the currently selected bone has changed */
+	UE_DEPRECATED(5.6, "Use RegisterOnSelectedBonesChanged instead.")
 	virtual FDelegateHandle RegisterOnSelectedBoneChanged(const FOnSelectedBoneChanged& Delegate) = 0;
 
 	/** Unregisters a delegate called when the currently selected bone has changed */
+	UE_DEPRECATED(5.6, "Use UnregisterOnSelectedBonesChanged instead.")
 	virtual void UnregisterOnSelectedBoneChanged(FDelegateHandle InHandle) = 0;
 
+	/** Registers a delegate called when the bones selection has changed */
+	virtual FDelegateHandle RegisterOnSelectedBonesChanged(const FOnSelectedBonesChanged& InDelegate) = 0;
+
+	/** Unregisters a delegate called when the bones selection has changed */
+	virtual void UnregisterOnSelectedBonesChanged(FDelegateHandle InHandle) = 0;
+	
 	/** Registers a delegate to be called when the currently selected socket has changed */
 	virtual FDelegateHandle RegisterOnSelectedSocketChanged(const FOnSelectedSocketChanged& Delegate) = 0;
 
diff --git a/Engine/Source/Editor/Persona/Public/SkeletonSelectionEditMode.h b/Engine/Source/Editor/Persona/Public/SkeletonSelectionEditMode.h
--- a/Engine/Source/Editor/Persona/Public/SkeletonSelectionEditMode.h
+++ b/Engine/Source/Editor/Persona/Public/SkeletonSelectionEditMode.h
@@ -43,7 +43,8 @@ public:
 	virtual bool GetCustomInputCoordinateSystem(FMatrix& InMatrix, void* InData) override;
 	virtual bool HandleClick(FEditorViewportClient* InViewportClient, HHitProxy *HitProxy, const FViewportClick &Click) override;
 	virtual bool CanCycleWidgetMode() const override;
-
+	virtual bool FrustumSelect(const FConvexVolume& InFrustum, FEditorViewportClient* InViewportClient, bool InSelect = true) override;
+	
 	virtual bool BeginTransform(const FGizmoState& InState) override;
 	virtual bool EndTransform(const FGizmoState& InState) override;
 	
diff --git a/Engine/Source/Editor/SkeletonEditor/Private/SSkeletonTree.cpp b/Engine/Source/Editor/SkeletonEditor/Private/SSkeletonTree.cpp
--- a/Engine/Source/Editor/SkeletonEditor/Private/SSkeletonTree.cpp
+++ b/Engine/Source/Editor/SkeletonEditor/Private/SSkeletonTree.cpp
@@ -189,13 +189,16 @@ void SSkeletonTree::Construct(const FArguments& InArgs, const TSharedRef<FEditab
 
 	if(PreviewScene.IsValid())
 	{
-		PreviewScene.Pin()->RegisterOnLODChanged(FSimpleDelegate::CreateSP(this, &SSkeletonTree::OnLODSwitched));
-		PreviewScene.Pin()->RegisterOnPreviewMeshChanged(FOnPreviewMeshChanged::CreateSP(this, &SSkeletonTree::OnPreviewMeshChanged));
-		PreviewScene.Pin()->RegisterOnSelectedBoneChanged(FOnSelectedBoneChanged::CreateSP(this, &SSkeletonTree::HandleSelectedBoneChanged));
-		PreviewScene.Pin()->RegisterOnSelectedSocketChanged(FOnSelectedSocketChanged::CreateSP(this, &SSkeletonTree::HandleSelectedSocketChanged));
-		PreviewScene.Pin()->RegisterOnDeselectAll(FSimpleDelegate::CreateSP(this, &SSkeletonTree::HandleDeselectAll));
+		if (TSharedPtr<IPersonaPreviewScene> PreviewScenePtr = PreviewScene.Pin())
+		{
+			PreviewScenePtr->RegisterOnLODChanged(FSimpleDelegate::CreateSP(this, &SSkeletonTree::OnLODSwitched));
+			PreviewScenePtr->RegisterOnPreviewMeshChanged(FOnPreviewMeshChanged::CreateSP(this, &SSkeletonTree::OnPreviewMeshChanged));
+			PreviewScenePtr->RegisterOnSelectedBonesChanged(FOnSelectedBonesChanged::CreateSP(this, &SSkeletonTree::HandleSelectedBonesChanged));
+			PreviewScenePtr->RegisterOnSelectedSocketChanged(FOnSelectedSocketChanged::CreateSP(this, &SSkeletonTree::HandleSelectedSocketChanged));
+			PreviewScenePtr->RegisterOnDeselectAll(FSimpleDelegate::CreateSP(this, &SSkeletonTree::HandleDeselectAll));
 
-		RegisterOnSelectionChanged(FOnSkeletonTreeSelectionChanged::CreateRaw(PreviewScene.Pin().Get(), &IPersonaPreviewScene::HandleSkeletonTreeSelectionChanged));
+			RegisterOnSelectionChanged(FOnSkeletonTreeSelectionChanged::CreateRaw(PreviewScenePtr.Get(), &IPersonaPreviewScene::HandleSkeletonTreeSelectionChanged));
+		}
 	}
 
 	if (InSkeletonTreeArgs.OnSelectionChanged.IsBound())
@@ -1890,6 +1893,44 @@ void SSkeletonTree::SetSelectedBone( const FName& BoneName, ESelectInfo::Type In
 	}
 }
 
+void SSkeletonTree::SetSelectedBones(const TArray<FName>& InBoneNames, ESelectInfo::Type InSelectInfo)
+{
+	if (InBoneNames.IsEmpty())
+	{
+		return;
+	}
+	
+	if (!bSelecting)
+	{
+		TGuardValue<bool> RecursionGuard(bSelecting, true);
+
+		const TArray<TSharedPtr<ISkeletonTreeItem>> SkeletonTreeItems = LinearItems.FilterByPredicate([&InBoneNames](const TSharedPtr<ISkeletonTreeItem>& SkeletonRow)
+		{
+			if (SkeletonRow->GetFilterResult() == ESkeletonTreeFilterResult::Hidden)
+			{
+				return false;
+			}
+
+			if (SkeletonRow->IsOfType<FSkeletonTreeBoneItem>() || SkeletonRow->IsOfType<FSkeletonTreeVirtualBoneItem>())
+			{
+				return InBoneNames.Contains(SkeletonRow->GetRowItemName());	
+			}
+
+			return false;
+		});
+
+		if (SkeletonTreeItems.IsEmpty())
+		{
+			return;
+		}
+
+		ExpandTreeOnSelection(SkeletonTreeItems[0]);
+		constexpr bool bSelected = true;
+		SkeletonTreeView->SetItemSelection(SkeletonTreeItems, bSelected, InSelectInfo);
+		SkeletonTreeView->RequestScrollIntoView(SkeletonTreeItems[0]);
+	}
+}
+
 void SSkeletonTree::DeselectAll()
 {
 	if (!bSelecting)
@@ -2566,9 +2607,9 @@ ESkeletonTreeFilterResult SSkeletonTree::HandleFilterSkeletonTreeItem(const FSke
 	return Result;
 }
 
-void SSkeletonTree::HandleSelectedBoneChanged(const FName& InBoneName, ESelectInfo::Type InSelectInfo)
+void SSkeletonTree::HandleSelectedBonesChanged(const TArray<FName>& InBoneNames, ESelectInfo::Type InSelectInfo)
 {
-	SetSelectedBone(InBoneName, InSelectInfo);
+	SetSelectedBones(InBoneNames, InSelectInfo);
 }
 
 void SSkeletonTree::HandleSelectedSocketChanged(const FSelectedSocketInfo& InSocketInfo)
diff --git a/Engine/Source/Editor/SkeletonEditor/Private/SSkeletonTree.h b/Engine/Source/Editor/SkeletonEditor/Private/SSkeletonTree.h
--- a/Engine/Source/Editor/SkeletonEditor/Private/SSkeletonTree.h
+++ b/Engine/Source/Editor/SkeletonEditor/Private/SSkeletonTree.h
@@ -77,6 +77,7 @@ public:
 	virtual void SetSkeletalMesh(USkeletalMesh* NewSkeletalMesh) override;
 	virtual void SetSelectedSocket(const struct FSelectedSocketInfo& InSocketInfo) override;
 	virtual void SetSelectedBone(const FName& InBoneName, ESelectInfo::Type InSelectInfo) override;
+	virtual void SetSelectedBones(const TArray<FName>& InBoneNames, ESelectInfo::Type InSelectInfo) override;
 	virtual void DeselectAll() override;
 	virtual TArray<TSharedPtr<ISkeletonTreeItem>> GetSelectedItems() const override { return SkeletonTreeView->GetSelectedItems(); }
 	virtual void SelectItemsBy(TFunctionRef<bool(const TSharedRef<ISkeletonTreeItem>&, bool&)> Predicate, const ESelectInfo::Type SelectInfo) const override;
@@ -324,6 +325,7 @@ private:
 
 	/** Handle bone selection changing externally */
 	void HandleSelectedBoneChanged(const FName& InBoneName, ESelectInfo::Type InSelectInfo);
+	void HandleSelectedBonesChanged(const TArray<FName>& InBoneNames, ESelectInfo::Type InSelectInfo);
 
 	/** Handle socket selection changing externally */
 	void HandleSelectedSocketChanged(const FSelectedSocketInfo& InSocketInfo);
diff --git a/Engine/Source/Editor/SkeletonEditor/Public/ISkeletonTree.h b/Engine/Source/Editor/SkeletonEditor/Public/ISkeletonTree.h
--- a/Engine/Source/Editor/SkeletonEditor/Public/ISkeletonTree.h
+++ b/Engine/Source/Editor/SkeletonEditor/Public/ISkeletonTree.h
@@ -132,6 +132,9 @@ public:
 
 	/** Set the selected bone */
 	virtual void SetSelectedBone(const FName& InBoneName, ESelectInfo::Type InSelectInfo) = 0;
+	
+	/** Set the selected bone(s) */
+	virtual void SetSelectedBones(const TArray<FName>& InBoneNames, ESelectInfo::Type InSelectInfo) = 0;
 
 	/** Deselect everything that is currently selected */
 	virtual void DeselectAll() = 0;
