Skeletal Mesh: Don't invalidate the GUID directly when we reduce bone counts. That should be done upstream from the bone reduction call, otherwise LOD reduction with bone removal will cause a GUID invalidation at cook time.

diff --git a/Engine/Source/Developer/MeshBoneReduction/Private/MeshBoneReduction.cpp b/Engine/Source/Developer/MeshBoneReduction/Private/MeshBoneReduction.cpp
--- a/Engine/Source/Developer/MeshBoneReduction/Private/MeshBoneReduction.cpp
+++ b/Engine/Source/Developer/MeshBoneReduction/Private/MeshBoneReduction.cpp
@@ -520,8 +520,6 @@ public:
 			}
 
 			delete SrcModel;
-
-			SkeletalMesh->InvalidateDeriveDataCacheGUID();
 		}
 		else
 		{
diff --git a/Engine/Source/Editor/SkeletonEditor/Private/SSkeletonTree.cpp b/Engine/Source/Editor/SkeletonEditor/Private/SSkeletonTree.cpp
--- a/Engine/Source/Editor/SkeletonEditor/Private/SSkeletonTree.cpp
+++ b/Engine/Source/Editor/SkeletonEditor/Private/SSkeletonTree.cpp
@@ -65,6 +65,7 @@
 #include "HAL/PlatformApplicationMisc.h"
 #include "Widgets/Views/STreeView.h"
 #include "IPinnedCommandList.h"
+#include "MeshBoneReduction.h"
 #include "PersonaModule.h"
 #include "SPositiveActionButton.h"
 #include "ToolMenus.h"
@@ -1115,7 +1116,8 @@ void SSkeletonTree::RemoveFromLOD(int32 LODIndex, bool bIncludeSelected, bool bI
 	}
 
 	UDebugSkelMeshComponent* PreviewMeshComponent = GetPreviewScene()->GetPreviewMeshComponent();
-	if (!PreviewMeshComponent->GetSkeletalMeshAsset())
+	USkeletalMesh *SkeletalMesh = PreviewMeshComponent->GetSkeletalMeshAsset();
+	if (!SkeletalMesh)
 	{
 		return;
 	}
@@ -1132,7 +1134,7 @@ void SSkeletonTree::RemoveFromLOD(int32 LODIndex, bool bIncludeSelected, bool bI
 
 		//Scoped post edit change
 		{
-			FScopedSkeletalMeshPostEditChange ScopedPostEditChange(PreviewMeshComponent->GetSkeletalMeshAsset());
+			FScopedSkeletalMeshPostEditChange ScopedPostEditChange(SkeletalMesh);
 
 			for (const TSharedPtr<FSkeletonTreeBoneItem>& Item : TreeSelection.GetSelectedItems<FSkeletonTreeBoneItem>())
 			{
@@ -1160,24 +1162,30 @@ void SSkeletonTree::RemoveFromLOD(int32 LODIndex, bool bIncludeSelected, bool bI
 				}
 			}
 
-			int32 TotalLOD = PreviewMeshComponent->GetSkeletalMeshAsset()->GetLODNum();
-			IMeshUtilities& MeshUtilities = FModuleManager::Get().LoadModuleChecked<IMeshUtilities>("MeshUtilities");
-
+			IMeshBoneReductionModule& MeshBoneReductionModule = FModuleManager::Get().LoadModuleChecked<IMeshBoneReductionModule>("MeshBoneReduction");
+			IMeshBoneReduction& MeshBoneReduction = *MeshBoneReductionModule.GetMeshBoneReductionInterface();
+			constexpr bool bCallPostEditChange = false;
+			
 			if (bIncludeBelowLODs)
 			{
+				const int32 TotalLOD = SkeletalMesh->GetLODNum();
 				for (int32 Index = LODIndex + 1; Index < TotalLOD; ++Index)
 				{
-					PreviewMeshComponent->GetSkeletalMeshAsset()->AddBoneToReductionSetting(Index, BonesToRemove);
+					SkeletalMesh->AddBoneToReductionSetting(Index, BonesToRemove);
 					// We don't pass BoneNamesToRemove, as AddBoneToReductionSetting has added them to the LODInfoArray[LODIndex].BonesToRemove
 					// Which will be used by RemoveBonesFromMesh if we pass null BonesNamesToRemove (else will just remove the newly deleted bones, which is wrong)
-					MeshUtilities.RemoveBonesFromMesh(PreviewMeshComponent->GetSkeletalMeshAsset(), Index, nullptr);
+					MeshBoneReduction.ReduceBoneCounts(SkeletalMesh, Index, nullptr, bCallPostEditChange);
 				}
 			}
 
 			// remove from current LOD
 			// We don't pass BoneNamesToRemove, as AddBoneToReductionSetting has added them to the LODInfoArray[LODIndex].BonesToRemove
 			// Which will be used by RemoveBonesFromMesh if we pass null BonesNamesToRemove (else will just remove the newly deleted bones, which is wrong)
-			MeshUtilities.RemoveBonesFromMesh(PreviewMeshComponent->GetSkeletalMeshAsset(), LODIndex, nullptr);
+			MeshBoneReduction.ReduceBoneCounts(SkeletalMesh, LODIndex, nullptr, bCallPostEditChange);
+			
+			// Ensure the mesh gets regenerated. The bone hierarchy should be a part of the DDC key, but isn't, so we forcibly update the
+			// base GUID instead.
+			SkeletalMesh->InvalidateDeriveDataCacheGUID();
 		}
 		// update UI to reflect the change
 		OnLODSwitched();
diff --git a/Engine/Source/Editor/SkeletonEditor/SkeletonEditor.Build.cs b/Engine/Source/Editor/SkeletonEditor/SkeletonEditor.Build.cs
--- a/Engine/Source/Editor/SkeletonEditor/SkeletonEditor.Build.cs
+++ b/Engine/Source/Editor/SkeletonEditor/SkeletonEditor.Build.cs
@@ -29,7 +29,7 @@ public class SkeletonEditor : ModuleRules
 				"EditorWidgets",
 				"ToolWidgets",
 				"AnimationCore",
-				"AnimationWidgets",
+				"AnimationWidgets", 
 			}
 		);
 
@@ -43,12 +43,14 @@ public class SkeletonEditor : ModuleRules
         PrivateIncludePathModuleNames.AddRange(
             new string[] {
                 "PropertyEditor",
+                "MeshBoneReduction",
             }
         );
 
         DynamicallyLoadedModuleNames.AddRange(
 			new string[] {
 				"PropertyEditor",
+				"MeshBoneReduction",
 			}
 		);
 	}
