Factoring the detection code for unused bones out of FEditableSkeleton::RemoveUnusedBones() so it can be called independently.

diff --git a/Engine/Source/Editor/SkeletonEditor/Private/EditableSkeleton.cpp b/Engine/Source/Editor/SkeletonEditor/Private/EditableSkeleton.cpp
--- a/Engine/Source/Editor/SkeletonEditor/Private/EditableSkeleton.cpp
+++ b/Engine/Source/Editor/SkeletonEditor/Private/EditableSkeleton.cpp
@@ -1452,102 +1452,122 @@ void FEditableSkeleton::RemoveCompatibleSkeleton(const USkeleton* InCompatibleSk
 void FEditableSkeleton::RemoveUnusedBones()
 {
 	TArray<FName> SkeletonBones;
-	const FReferenceSkeleton& RefSkeleton = Skeleton->GetReferenceSkeleton();
+	FindUnusedBones(SkeletonBones);
 
-	for (int32 BoneIndex = 0; BoneIndex < RefSkeleton.GetRawBoneNum(); ++BoneIndex)
+	if (SkeletonBones.Num() > 0)
+	{
+		const FText RemoveBoneMessage = FText::Format(LOCTEXT("RemoveBoneWarning", "Continuing will remove the following bones from the skeleton '{0}'. These bones are not being used by any of the SkeletalMeshes assigned to this skeleton.\n\nOnce the bones have been removed, all loaded animations for this skeleton will be recompressed (any that aren't loaded will be recompressed the next time they are loaded)."), FText::FromString(Skeleton->GetName()));
+
+		// Ask User whether they would like to remove the bones from the skeleton
+		if (SSkeletonBoneRemoval::ShowModal(SkeletonBones, RemoveBoneMessage))
+		{
+			//Remove these bones from the skeleton
+			Skeleton->RemoveBonesFromSkeleton(SkeletonBones, true);
+		}
+	}
+	else
 	{
-		SkeletonBones.Add(RefSkeleton.GetBoneName(BoneIndex));
+		FMessageDialog::Open(EAppMsgType::Ok, LOCTEXT("NoBonesToRemove", "No unused bones were found."));
 	}
 
+	OnTreeRefresh.Broadcast();
+}
+
+void FEditableSkeleton::FindUnusedBones(TArray<FName>& OutSkeletonBones, bool bSilent /*=false*/)
+{
+	// This operation can take a long time, so unless we're running silently, give the user a heads-up
+	if (!bSilent)
+	{
+		FARFilter Filter;
+		Filter.ClassPaths.Add(USkeletalMesh::StaticClass()->GetClassPathName());
+
+		FString SkeletonString = FAssetData(Skeleton).GetExportTextName();
+		Filter.TagsAndValues.Add(USkeletalMesh::GetSkeletonMemberName(), SkeletonString);
+
+		TArray<FAssetData> SkeletalMeshes;
+		FAssetRegistryModule& AssetRegistryModule = FModuleManager::LoadModuleChecked<FAssetRegistryModule>(TEXT("AssetRegistry"));
+		AssetRegistryModule.Get().GetAssets(Filter, SkeletalMeshes);
+
+		FText TimeTakenMessage = FText::Format(LOCTEXT("TimeTakenWarning", "In order to verify bone use, all Skeletal Meshes that use this skeleton will be loaded.  This may take some time.\n\nProceed?\n\nNumber of Meshes: {0}"), FText::AsNumber(SkeletalMeshes.Num()));
+
+		if (FMessageDialog::Open(EAppMsgType::YesNo, TimeTakenMessage) != EAppReturnType::Yes)
+		{
+			return;
+		}
+	}
+
+	const FReferenceSkeleton& RefSkeleton = Skeleton->GetReferenceSkeleton();
+	
+	for (int32 BoneIndex = 0; BoneIndex < RefSkeleton.GetRawBoneNum(); ++BoneIndex)
+	{
+		OutSkeletonBones.Add(RefSkeleton.GetBoneName(BoneIndex));
+	}
+	
 	FARFilter Filter;
 	Filter.ClassPaths.Add(USkeletalMesh::StaticClass()->GetClassPathName());
-
+	
 	FString SkeletonString = FAssetData(Skeleton).GetExportTextName();
 	Filter.TagsAndValues.Add(USkeletalMesh::GetSkeletonMemberName(), SkeletonString);
-
+	
 	TArray<FAssetData> SkeletalMeshes;
 	FAssetRegistryModule& AssetRegistryModule = FModuleManager::LoadModuleChecked<FAssetRegistryModule>(TEXT("AssetRegistry"));
 	AssetRegistryModule.Get().GetAssets(Filter, SkeletalMeshes);
-
-	FText TimeTakenMessage = FText::Format(LOCTEXT("TimeTakenWarning", "In order to verify bone use all Skeletal Meshes that use this skeleton will be loaded, this may take some time.\n\nProceed?\n\nNumber of Meshes: {0}"), FText::AsNumber(SkeletalMeshes.Num()));
-
-	if (FMessageDialog::Open(EAppMsgType::YesNo, TimeTakenMessage) == EAppReturnType::Yes)
+	
+	const FText StatusUpdate = FText::Format(LOCTEXT("RemoveUnusedBones_ProcessingAssetsFor", "Processing Skeletal Meshes for {0}"), FText::FromString(Skeleton->GetName()));
+	GWarn->BeginSlowTask(StatusUpdate, true);
+	
+	// Loop through virtual bones and remove the bones they use from the list
+	for (const FVirtualBone& VB : Skeleton->GetVirtualBones())
 	{
-		const FText StatusUpdate = FText::Format(LOCTEXT("RemoveUnusedBones_ProcessingAssetsFor", "Processing Skeletal Meshes for {0}"), FText::FromString(Skeleton->GetName()));
-		GWarn->BeginSlowTask(StatusUpdate, true);
-
-		// Loop through virtual bones and remove the bones they use from the list
-		for (const FVirtualBone& VB : Skeleton->GetVirtualBones())
+		OutSkeletonBones.Remove(VB.SourceBoneName);
+		OutSkeletonBones.Remove(VB.TargetBoneName);
+		if (OutSkeletonBones.Num() == 0)
 		{
-			SkeletonBones.Remove(VB.SourceBoneName);
-			SkeletonBones.Remove(VB.TargetBoneName);
-			if (SkeletonBones.Num() == 0)
-			{
-				break;
-			}
+			break;
 		}
-
-		if (SkeletonBones.Num() != 0)
+	}
+	
+	if (OutSkeletonBones.Num() != 0)
+	{
+		// Loop through all SkeletalMeshes and remove the bones they use from our list
+		for (int32 MeshIdx = 0; MeshIdx < SkeletalMeshes.Num(); ++MeshIdx)
 		{
-			// Loop through all SkeletalMeshes and remove the bones they use from our list
-			for (int32 MeshIdx = 0; MeshIdx < SkeletalMeshes.Num(); ++MeshIdx)
+			GWarn->StatusUpdate(MeshIdx, SkeletalMeshes.Num(), StatusUpdate);
+	
+			USkeletalMesh* Mesh = Cast<USkeletalMesh>(SkeletalMeshes[MeshIdx].GetAsset());
+			const FReferenceSkeleton& MeshRefSkeleton = Mesh->GetRefSkeleton();
+	
+			for (int32 BoneIndex = 0; BoneIndex < MeshRefSkeleton.GetRawBoneNum(); ++BoneIndex)
 			{
-				GWarn->StatusUpdate(MeshIdx, SkeletalMeshes.Num(), StatusUpdate);
-
-				USkeletalMesh* Mesh = Cast<USkeletalMesh>(SkeletalMeshes[MeshIdx].GetAsset());
-				const FReferenceSkeleton& MeshRefSkeleton = Mesh->GetRefSkeleton();
-
-				for (int32 BoneIndex = 0; BoneIndex < MeshRefSkeleton.GetRawBoneNum(); ++BoneIndex)
-				{
-					SkeletonBones.Remove(MeshRefSkeleton.GetBoneName(BoneIndex));
-					if (SkeletonBones.Num() == 0)
-					{
-						break;
-					}
-				}
-				if (SkeletonBones.Num() == 0)
+				OutSkeletonBones.Remove(MeshRefSkeleton.GetBoneName(BoneIndex));
+				if (OutSkeletonBones.Num() == 0)
 				{
 					break;
 				}
 			}
-		}
-
-		GWarn->EndSlowTask();
-
-		//Remove bones that are a parent to bones we aren't removing
-		for (int32 BoneIndex = RefSkeleton.GetRawBoneNum() - 1; BoneIndex >= 0; --BoneIndex)
-		{
-			FName CurrBoneName = RefSkeleton.GetBoneName(BoneIndex);
-			if (!SkeletonBones.Contains(CurrBoneName))
+			if (OutSkeletonBones.Num() == 0)
 			{
-				//We aren't removing this bone, so remove parent from list of bones to remove too
-				int32 ParentIndex = RefSkeleton.GetParentIndex(BoneIndex);
-				if (ParentIndex != INDEX_NONE)
-				{
-					SkeletonBones.Remove(RefSkeleton.GetBoneName(ParentIndex));
-				}
+				break;
 			}
 		}
-
-		// If we have any bones left they are unused
-		if (SkeletonBones.Num() > 0)
+	}
+	
+	GWarn->EndSlowTask();
+	
+	//Remove bones that are a parent to bones we aren't removing
+	for (int32 BoneIndex = RefSkeleton.GetRawBoneNum() - 1; BoneIndex >= 0; --BoneIndex)
+	{
+		FName CurrBoneName = RefSkeleton.GetBoneName(BoneIndex);
+		if (!OutSkeletonBones.Contains(CurrBoneName))
 		{
-			const FText RemoveBoneMessage = FText::Format(LOCTEXT("RemoveBoneWarning", "Continuing will remove the following bones from the skeleton '{0}'. These bones are not being used by any of the SkeletalMeshes assigned to this skeleton\n\nOnce the bones have been removed all loaded animations for this skeleton will be recompressed (any that aren't loaded will be recompressed the next time they are loaded)."), FText::FromString(Skeleton->GetName()));
-
-			// Ask User whether they would like to remove the bones from the skeleton
-			if (SSkeletonBoneRemoval::ShowModal(SkeletonBones, RemoveBoneMessage))
+			//We aren't removing this bone, so remove parent from list of bones to remove too
+			int32 ParentIndex = RefSkeleton.GetParentIndex(BoneIndex);
+			if (ParentIndex != INDEX_NONE)
 			{
-				//Remove these bones from the skeleton
-				Skeleton->RemoveBonesFromSkeleton(SkeletonBones, true);
+				OutSkeletonBones.Remove(RefSkeleton.GetBoneName(ParentIndex));
 			}
 		}
-		else
-		{
-			FMessageDialog::Open(EAppMsgType::Ok, LOCTEXT("NoBonesToRemove", "No unused bones were found."));
-		}
 	}
-
-	OnTreeRefresh.Broadcast();
 }
 
 void FEditableSkeleton::UpdateSkeletonReferencePose(USkeletalMesh* InSkeletalMesh)
diff --git a/Engine/Source/Editor/SkeletonEditor/Private/EditableSkeleton.h b/Engine/Source/Editor/SkeletonEditor/Private/EditableSkeleton.h
--- a/Engine/Source/Editor/SkeletonEditor/Private/EditableSkeleton.h
+++ b/Engine/Source/Editor/SkeletonEditor/Private/EditableSkeleton.h
@@ -62,6 +62,7 @@ public:
 	virtual void AddCompatibleSkeleton(const USkeleton* InCompatibleSkeleton) override;
 	virtual void RemoveCompatibleSkeleton(const USkeleton* InCompatibleSkeleton) override;
 	virtual void RemoveUnusedBones() override;
+	virtual void FindUnusedBones(TArray<FName>& OutUnusedBones, bool bSilent = false) override;
 	virtual void UpdateSkeletonReferencePose(USkeletalMesh* InSkeletalMesh) override;
 	virtual void RegisterSlotNode(const FName& InSlotName) override;
 	virtual bool AddSlotGroupName(const FName& InSlotName) override;
diff --git a/Engine/Source/Editor/SkeletonEditor/Public/IEditableSkeleton.h b/Engine/Source/Editor/SkeletonEditor/Public/IEditableSkeleton.h
--- a/Engine/Source/Editor/SkeletonEditor/Public/IEditableSkeleton.h
+++ b/Engine/Source/Editor/SkeletonEditor/Public/IEditableSkeleton.h
@@ -172,6 +172,9 @@ public:
 	/** Remove any bones that are not used by any skeletal meshes */
 	virtual void RemoveUnusedBones() = 0;
 
+	/** Returns a list of all bones that are not used by any skeletal meshes */
+	virtual void FindUnusedBones(TArray<FName>& OutUnusedBones, bool bSilent = false) = 0;
+
 	/** Create RefLocalPoses from InSkeletalMesh */
 	virtual void UpdateSkeletonReferencePose(class USkeletalMesh* InSkeletalMesh) = 0;
 
