Makes PhAT resilient to the skeleton being deleted

#Synced-CL 49327860

diff --git a/Engine/Source/Editor/SkeletonEditor/Private/SkeletonTreeBuilder.cpp b/Engine/Source/Editor/SkeletonEditor/Private/SkeletonTreeBuilder.cpp
--- a/Engine/Source/Editor/SkeletonEditor/Private/SkeletonTreeBuilder.cpp
+++ b/Engine/Source/Editor/SkeletonEditor/Private/SkeletonTreeBuilder.cpp
@@ -205,8 +205,14 @@ bool FSkeletonTreeBuilder::IsShowingAttachedAssets() const
 
 void FSkeletonTreeBuilder::AddBones(FSkeletonTreeBuilderOutput& Output)
 {
-	const USkeleton& Skeleton = EditableSkeletonPtr.Pin()->GetSkeleton();
+	TSharedPtr<IEditableSkeleton> EditableSkeleton = EditableSkeletonPtr.Pin();
+	if (!EditableSkeleton || !EditableSkeleton->IsSkeletonValid())
+	{
+		UE_LOG(LogAnimation, Warning, TEXT("Skeleton tree builder - skeleton is invalid"));
+		return;
+	}
 
+	const USkeleton& Skeleton = EditableSkeleton->GetSkeleton();
 	const FReferenceSkeleton& RefSkeleton = Skeleton.GetReferenceSkeleton();
 
 	struct FBoneInfo
@@ -295,7 +301,13 @@ void FSkeletonTreeBuilder::AddBones(FSkeletonTreeBuilderOutput& Output)
 
 void FSkeletonTreeBuilder::AddSockets(FSkeletonTreeBuilderOutput& Output)
 {
-	const USkeleton& Skeleton = EditableSkeletonPtr.Pin()->GetSkeleton();
+	TSharedPtr<IEditableSkeleton> EditableSkeleton = EditableSkeletonPtr.Pin();
+	if (!EditableSkeleton || !EditableSkeleton->IsSkeletonValid())
+	{
+		UE_LOG(LogAnimation, Warning, TEXT("Skeleton tree builder - skeleton is invalid"));
+		return;
+	}
+	const USkeleton& Skeleton = EditableSkeleton->GetSkeleton();
 
 	// Add the sockets for the skeleton
 	AddSocketsFromData(Skeleton.Sockets, ESocketParentType::Skeleton, Output);
@@ -313,7 +325,14 @@ void FSkeletonTreeBuilder::AddSockets(FSkeletonTreeBuilderOutput& Output)
 
 void FSkeletonTreeBuilder::AddAttachedAssets(FSkeletonTreeBuilderOutput& Output)
 {
-	const USkeleton& Skeleton = EditableSkeletonPtr.Pin()->GetSkeleton();
+	TSharedPtr<IEditableSkeleton> EditableSkeleton = EditableSkeletonPtr.Pin();
+	if (!EditableSkeleton || !EditableSkeleton->IsSkeletonValid())
+	{
+		UE_LOG(LogAnimation, Warning, TEXT("Skeleton tree builder - skeleton is invalid"));
+		return;
+	}
+
+	const USkeleton& Skeleton = EditableSkeleton->GetSkeleton();
 
 	// Mesh attached items...
 	if (PreviewScenePtr.IsValid())
@@ -331,7 +350,14 @@ void FSkeletonTreeBuilder::AddAttachedAssets(FSkeletonTreeBuilderOutput& Output)
 
 void FSkeletonTreeBuilder::AddSocketsFromData(const TArray< USkeletalMeshSocket* >& SocketArray, ESocketParentType ParentType, FSkeletonTreeBuilderOutput& Output)
 {
-	const USkeleton& Skeleton = EditableSkeletonPtr.Pin()->GetSkeleton();
+	TSharedPtr<IEditableSkeleton> EditableSkeleton = EditableSkeletonPtr.Pin();
+	if (!EditableSkeleton || !EditableSkeleton->IsSkeletonValid())
+	{
+		UE_LOG(LogAnimation, Warning, TEXT("Skeleton tree builder - skeleton is invalid"));
+		return;
+	}
+
+	const USkeleton& Skeleton = EditableSkeleton->GetSkeleton();
 
 	for (auto SocketIt = SocketArray.CreateConstIterator(); SocketIt; ++SocketIt)
 	{
@@ -341,7 +367,7 @@ void FSkeletonTreeBuilder::AddSocketsFromData(const TArray< USkeletalMeshSocket*
 
 		if (ParentType == ESocketParentType::Mesh)
 		{
-			bIsCustomized = EditableSkeletonPtr.Pin()->DoesSocketAlreadyExist(nullptr, FText::FromName(Socket->SocketName), ESocketParentType::Skeleton, nullptr);
+			bIsCustomized = EditableSkeleton->DoesSocketAlreadyExist(nullptr, FText::FromName(Socket->SocketName), ESocketParentType::Skeleton, nullptr);
 		}
 		else
 		{
@@ -350,7 +376,7 @@ void FSkeletonTreeBuilder::AddSocketsFromData(const TArray< USkeletalMeshSocket*
 				UDebugSkelMeshComponent* PreviewMeshComponent = PreviewScenePtr.Pin()->GetPreviewMeshComponent();
 				if (USkeletalMesh* const SkeletalMesh = PreviewMeshComponent->GetSkeletalMeshAsset())
 				{
-					bIsCustomized = EditableSkeletonPtr.Pin()->DoesSocketAlreadyExist(nullptr, FText::FromName(Socket->SocketName), ESocketParentType::Mesh, SkeletalMesh);
+					bIsCustomized = EditableSkeleton->DoesSocketAlreadyExist(nullptr, FText::FromName(Socket->SocketName), ESocketParentType::Mesh, SkeletalMesh);
 				}
 			}
 		}
@@ -371,7 +397,14 @@ void FSkeletonTreeBuilder::AddAttachedAssetContainer(const FPreviewAssetAttachCo
 
 void FSkeletonTreeBuilder::AddVirtualBones(FSkeletonTreeBuilderOutput& Output)
 {
-	const TArray<FVirtualBone>& VirtualBones = EditableSkeletonPtr.Pin()->GetSkeleton().GetVirtualBones();
+	TSharedPtr<IEditableSkeleton> EditableSkeleton = EditableSkeletonPtr.Pin();
+	if (!EditableSkeleton || !EditableSkeleton->IsSkeletonValid())
+	{
+		UE_LOG(LogAnimation, Warning, TEXT("Skeleton tree builder - skeleton is invalid"));
+		return;
+	}
+
+	const TArray<FVirtualBone>& VirtualBones = EditableSkeleton->GetSkeleton().GetVirtualBones();
 	for (const FVirtualBone& VirtualBone : VirtualBones)
 	{
 		Output.Add(CreateVirtualBoneTreeItem(VirtualBone.VirtualBoneName), VirtualBone.SourceBoneName, { FSkeletonTreeBoneItem::GetTypeId(), FSkeletonTreeVirtualBoneItem::GetTypeId() }, /*bAddToHead*/ true);
