Makes AnimationEditorPreviewScene and the EditableSkeleton (more) robust to the Skeleton being deleted.

#Synced-CL 49327860

diff --git a/Engine/Source/Editor/Persona/Private/AnimationEditorPreviewScene.cpp b/Engine/Source/Editor/Persona/Private/AnimationEditorPreviewScene.cpp
--- a/Engine/Source/Editor/Persona/Private/AnimationEditorPreviewScene.cpp
+++ b/Engine/Source/Editor/Persona/Private/AnimationEditorPreviewScene.cpp
@@ -176,7 +176,8 @@ void FAnimationEditorPreviewScene::SetPreviewMeshComponent(UDebugSkelMeshCompone
 
 void FAnimationEditorPreviewScene::SetPreviewMesh(USkeletalMesh* NewPreviewMesh, bool bAllowOverrideBaseMesh)
 {
-	if (NewPreviewMesh != nullptr && GetEditableSkeleton().IsValid() && !GetEditableSkeleton()->GetSkeleton().IsCompatibleMesh(NewPreviewMesh))
+	if (NewPreviewMesh != nullptr && GetEditableSkeleton().IsValid() && GetEditableSkeleton()->IsSkeletonValid() && 
+		!GetEditableSkeleton()->GetSkeleton().IsCompatibleMesh(NewPreviewMesh))
 	{
 		const USkeleton& Skeleton = GetEditableSkeleton()->GetSkeleton();
 		const bool bSkipSkeletonCompatibility = PersonaToolkit.IsValid() && PersonaToolkit.Pin()->CanPreviewMeshUseDifferentSkeleton(); 
@@ -411,7 +412,9 @@ void FAnimationEditorPreviewScene::RefreshAdditionalMeshes(bool bAllowOverrideBa
 	AdditionalMeshes.Empty();
 
 	// add new components
-	UDataAsset* PreviewSceneAdditionalMeshes = GetEditableSkeleton().IsValid() ? GetEditableSkeleton()->GetSkeleton().GetAdditionalPreviewSkeletalMeshes() : nullptr;
+	UDataAsset* PreviewSceneAdditionalMeshes = 
+		GetEditableSkeleton().IsValid() && GetEditableSkeleton()->IsSkeletonValid() 
+		? GetEditableSkeleton()->GetSkeleton().GetAdditionalPreviewSkeletalMeshes() : nullptr;
 	if (PreviewSceneAdditionalMeshes == nullptr)
 	{
 		PreviewSceneAdditionalMeshes = PreviewSceneDescription->DefaultAdditionalMeshes;
@@ -492,7 +495,7 @@ void FAnimationEditorPreviewScene::AddPreviewAttachedObjects()
 		}
 	}
 
-	if(GetEditableSkeleton().IsValid())
+	if(GetEditableSkeleton().IsValid() && GetEditableSkeleton()->IsSkeletonValid())
 	{
 		const USkeleton& Skeleton = GetEditableSkeleton()->GetSkeleton();
 
@@ -589,7 +592,7 @@ void FAnimationEditorPreviewScene::RemoveAttachedComponent( bool bRemovePreviewA
 
 	if( !bRemovePreviewAttached )
 	{
-		if(GetEditableSkeleton().IsValid())
+		if(GetEditableSkeleton().IsValid() && GetEditableSkeleton()->IsSkeletonValid())
 		{
 			const USkeleton& Skeleton = GetEditableSkeleton()->GetSkeleton();
 			for(auto Iter = Skeleton.PreviewAttachedAssetContainer.CreateConstIterator(); Iter; ++Iter)
@@ -678,7 +681,8 @@ void FAnimationEditorPreviewScene::SetPreviewAnimationAsset(UAnimationAsset* Ani
 {
 	if (SkeletalMeshComponent)
 	{
-		const USkeleton* Skeleton = GetEditableSkeleton().IsValid() ? &GetEditableSkeleton()->GetSkeleton() : nullptr;
+		const USkeleton* Skeleton = GetEditableSkeleton().IsValid() && GetEditableSkeleton()->IsSkeletonValid() 
+			? &GetEditableSkeleton()->GetSkeleton() : nullptr;
 
 		RemoveAttachedComponent(false);
 
@@ -870,7 +874,8 @@ void FAnimationEditorPreviewScene::ClearSelectedBone()
 
 void FAnimationEditorPreviewScene::SetSelectedBone(const FName& BoneName, ESelectInfo::Type InSelectInfo)
 {
-	const int32 BoneIndex = GetEditableSkeleton().IsValid() ? GetEditableSkeleton()->GetSkeleton().GetReferenceSkeleton().FindBoneIndex(BoneName) : INDEX_NONE;
+	const int32 BoneIndex = GetEditableSkeleton().IsValid() && GetEditableSkeleton()->IsSkeletonValid() 
+		? GetEditableSkeleton()->GetSkeleton().GetReferenceSkeleton().FindBoneIndex(BoneName) : INDEX_NONE;
 	if (BoneIndex != INDEX_NONE)
 	{
 		SetSelectedBones({BoneName}, InSelectInfo);
@@ -883,7 +888,7 @@ void FAnimationEditorPreviewScene::SetSelectedBone(const FName& BoneName, ESelec
 
 void FAnimationEditorPreviewScene::SetSelectedBones(const TArray<FName>& InBoneNames, ESelectInfo::Type InSelectInfo)
 {
-	if (!GetEditableSkeleton().IsValid())
+	if (!GetEditableSkeleton().IsValid() || !GetEditableSkeleton()->IsSkeletonValid())
 	{
 		return;
 	}
diff --git a/Engine/Source/Editor/PhysicsAssetEditor/Private/PhysicsAssetEditorAnimInstanceProxy.cpp b/Engine/Source/Editor/PhysicsAssetEditor/Private/PhysicsAssetEditorAnimInstanceProxy.cpp
--- a/Engine/Source/Editor/PhysicsAssetEditor/Private/PhysicsAssetEditorAnimInstanceProxy.cpp
+++ b/Engine/Source/Editor/PhysicsAssetEditor/Private/PhysicsAssetEditorAnimInstanceProxy.cpp
@@ -71,7 +71,10 @@ void FPhysicsAssetEditorAnimInstanceProxy::UpdateAnimationNode(const FAnimationU
 {
 	if (CurrentAsset != nullptr)
 	{
-		FAnimPreviewInstanceProxy::UpdateAnimationNode(InContext);
+		if (InContext.AnimInstanceProxy && InContext.AnimInstanceProxy->IsSkeletonValid())
+		{
+			FAnimPreviewInstanceProxy::UpdateAnimationNode(InContext);
+		}
 	}
 	else
 	{
diff --git a/Engine/Source/Editor/SkeletonEditor/Private/EditableSkeleton.cpp b/Engine/Source/Editor/SkeletonEditor/Private/EditableSkeleton.cpp
--- a/Engine/Source/Editor/SkeletonEditor/Private/EditableSkeleton.cpp
+++ b/Engine/Source/Editor/SkeletonEditor/Private/EditableSkeleton.cpp
@@ -135,11 +135,13 @@ private:
 FEditableSkeleton::FEditableSkeleton(USkeleton* InSkeleton)
 	: Skeleton(InSkeleton)
 {
+	check(Skeleton);
 	Skeleton->SetFlags(RF_Transactional);
 }
 
 const USkeleton& FEditableSkeleton::GetSkeleton() const
 {
+	check(Skeleton);
 	return *Skeleton;
 }
 
@@ -150,24 +152,25 @@ bool FEditableSkeleton::IsSkeletonValid() const
 
 const TArray<class UBlendProfile*>& FEditableSkeleton::GetBlendProfiles() const
 {
+	check(Skeleton);
 	return Skeleton->BlendProfiles;
 }
 
 UBlendProfile* FEditableSkeleton::GetBlendProfile(const FName& NameToUse)
 {
-	return Skeleton->GetBlendProfile(NameToUse);
+	return Skeleton ? Skeleton->GetBlendProfile(NameToUse) : nullptr;
 }
 
 UBlendProfile* FEditableSkeleton::CreateNewBlendProfile(const FName& NameToUse)
 {
 	FScopedTransaction Transaction(LOCTEXT("CreateBlendProfile", "Create Blend Profile"));
 
-	return Skeleton->CreateNewBlendProfile(NameToUse);
+	return Skeleton ? Skeleton->CreateNewBlendProfile(NameToUse) : nullptr;
 }
 
 void FEditableSkeleton::RemoveBlendProfile(UBlendProfile* InBlendProfile)
 {
-	if(InBlendProfile)
+	if(InBlendProfile && Skeleton)
 	{
 		FScopedTransaction Transaction(LOCTEXT("RemoveBlendProfile", "Remove Blend Profile"));
 
@@ -179,6 +182,11 @@ void FEditableSkeleton::RemoveBlendProfile(UBlendProfile* InBlendProfile)
 
 UBlendProfile* FEditableSkeleton::RenameBlendProfile(const FName& InBlendProfileName, const FName& InNewBlendProfileName)
 {
+	if (!Skeleton)
+	{
+		return nullptr;
+	}
+
 	FScopedTransaction Transaction(LOCTEXT("RenameBlendProfile", "Rename Blend Profile"));
 
 	UBlendProfile* Profile = Skeleton->RenameBlendProfile(InBlendProfileName, InNewBlendProfileName);
@@ -193,8 +201,12 @@ UBlendProfile* FEditableSkeleton::RenameBlendProfile(const FName& InBlendProfile
 
 void FEditableSkeleton::SetBlendProfileScale(const FName& InBlendProfileName, const FName& InBoneName, float InNewScale, bool bInRecurse)
 {
-	UBlendProfile* BlendProfile = GetBlendProfile(InBlendProfileName);
-	if (BlendProfile)
+	if (!Skeleton)
+	{
+		return;
+	}
+
+	if (UBlendProfile* BlendProfile = GetBlendProfile(InBlendProfileName))
 	{
 		FScopedTransaction Transaction(LOCTEXT("SetBlendProfileScale", "Set Blend Profile Scale"));
 
@@ -207,8 +219,7 @@ void FEditableSkeleton::SetBlendProfileScale(const FName& InBlendProfileName, co
 
 void FEditableSkeleton::SetBlendProfileMode(FName InBlendProfileName, EBlendProfileMode ProfileMode)
 {
-	UBlendProfile* BlendProfile = GetBlendProfile(InBlendProfileName);
-	if (BlendProfile)
+	if (UBlendProfile* BlendProfile = GetBlendProfile(InBlendProfileName))
 	{
 		FScopedTransaction Transaction(LOCTEXT("SetBlendProfileMode", "Set Blend Profile Mode"));
 
@@ -231,11 +242,16 @@ void FEditableSkeleton::SetBlendProfileMode(FName InBlendProfileName, EBlendProf
 
 void FEditableSkeleton::AddReferencedObjects(FReferenceCollector& Collector)
 {
+	// Note that this has the potential to set Skeleton to nullptr
 	Collector.AddReferencedObject(Skeleton);
 }
 
 void FEditableSkeleton::SetBoneTranslationRetargetingMode(FName InBoneName, EBoneTranslationRetargetingMode::Type NewRetargetingMode)
 {
+	if (!Skeleton)
+	{
+		return;
+	}
 	const FScopedTransaction Transaction(LOCTEXT("SetBoneTranslationRetargetingMode", "Set Bone Translation Retargeting Mode"));
 	Skeleton->Modify();
 
@@ -246,6 +262,10 @@ void FEditableSkeleton::SetBoneTranslationRetargetingMode(FName InBoneName, EBon
 
 EBoneTranslationRetargetingMode::Type FEditableSkeleton::GetBoneTranslationRetargetingMode(FName InBoneName) const
 {
+	if (!Skeleton)
+	{
+		return EBoneTranslationRetargetingMode::Animation; // Arbitrary return value
+	}
 	const int32 BoneIndex = Skeleton->GetReferenceSkeleton().FindBoneIndex(InBoneName);
 	return Skeleton->GetBoneTranslationRetargetingMode(BoneIndex);
 }
@@ -257,6 +277,11 @@ void FEditableSkeleton::RefreshBoneTree()
 
 bool FEditableSkeleton::DoesSocketAlreadyExist(const USkeletalMeshSocket* InSocket, const FText& InSocketName, ESocketParentType SocketParentType, USkeletalMesh* InSkeletalMesh) const
 {
+	if (!Skeleton)
+	{
+		return false;
+	}
+
 	TArray<TObjectPtr<USkeletalMeshSocket>>* SocketArrayPtr = nullptr;
 	if (SocketParentType == ESocketParentType::Mesh && InSkeletalMesh)
 	{
@@ -285,6 +310,12 @@ bool FEditableSkeleton::DoesSocketAlreadyExist(const USkeletalMeshSocket* InSock
 
 void FEditableSkeleton::GetAssetsContainingCurves(const FName& InContainerName, const TArray<FName>& InNames, TArray<FAssetData>& OutAssets) const
 {
+	if (!Skeleton)
+	{
+		OutAssets.Empty();
+		return;
+	}
+
 	FAssetData SkeletonData(Skeleton);
 	const FString CurrentSkeletonName = SkeletonData.GetExportTextName();
 
@@ -453,6 +484,11 @@ void FEditableSkeleton::RenameSocket(const FName OldSocketName, const FName NewS
 
 void FEditableSkeleton::SetSocketParent(const FName& SocketName, const FName& NewParentName, USkeletalMesh* InSkeletalMesh, bool bKeepAbsoluteLocation)
 {
+	if (!Skeleton)
+	{
+		return;
+	}
+
 	USkeletalMeshSocket* Socket = nullptr;
 	const FReferenceSkeleton* ReferenceSkeleton = nullptr;
 
@@ -506,6 +542,11 @@ void FEditableSkeleton::SetSocketParent(const FName& SocketName, const FName& Ne
 
 void FEditableSkeleton::HandlePasteSockets(const FName& InBoneName, USkeletalMesh* InSkeletalMesh)
 {
+	if (!Skeleton)
+	{
+		return;
+	}
+
 	FString PasteString;
 	FPlatformApplicationMisc::ClipboardPaste(PasteString);
 	const TCHAR* PastePtr = *PasteString;
@@ -541,6 +582,10 @@ USkeletalMeshSocket* FEditableSkeleton::AddSocket(const FName& InBoneName)
 
 USkeletalMeshSocket* FEditableSkeleton::HandleAddSocket(const FName& InBoneName)
 {
+	if (!Skeleton)
+	{
+		return nullptr;
+	}
 	const FScopedTransaction Transaction(LOCTEXT("AddSocket", "Add Socket to Skeleton"));
 	Skeleton->Modify();
 
@@ -563,6 +608,10 @@ bool FEditableSkeleton::HandleAddVirtualBone(const FName SourceBoneName, const F
 
 bool FEditableSkeleton::HandleAddVirtualBone(const FName SourceBoneName, const FName TargetBoneName, FName& NewVirtualBoneName)
 {
+	if (!Skeleton)
+	{
+		return false;
+	}
 	FScopedTransaction Transaction(LOCTEXT("AddVirtualBone", "Add Virtual Bone to Skeleton"));
 	const bool Success = Skeleton->AddNewVirtualBone(SourceBoneName, TargetBoneName, NewVirtualBoneName);
 	if (!Success)
@@ -598,6 +647,10 @@ void FEditableSkeleton::HandleCustomizeSocket(USkeletalMeshSocket* InSocketToCus
 
 void FEditableSkeleton::HandlePromoteSocket(USkeletalMeshSocket* InSocketToPromote)
 {
+	if (!Skeleton)
+	{
+		return;
+	}
 	const FScopedTransaction Transaction(LOCTEXT("PromoteSocket", "Promote Socket"));
 	Skeleton->Modify();
 
@@ -615,6 +668,10 @@ void FEditableSkeleton::HandlePromoteSocket(USkeletalMeshSocket* InSocketToPromo
 
 void FEditableSkeleton::HandleRemoveAllAssets(TSharedPtr<IPersonaPreviewScene> InPreviewScene)
 {
+	if (!Skeleton)
+	{
+		return;
+	}
 	FScopedTransaction Transaction(LOCTEXT("AttachedAssetRemoveUndo", "Remove All Attached Assets"));
 	Skeleton->Modify();
 
@@ -644,6 +701,10 @@ void FEditableSkeleton::DeleteAttachedObjects(FPreviewAssetAttachContainer& Atta
 
 USkeletalMeshSocket* FEditableSkeleton::DuplicateSocket(const FSelectedSocketInfo& SocketInfoToDuplicate, const FName& NewParentBoneName, USkeletalMesh* InSkeletalMesh)
 {
+	if (!Skeleton)
+	{
+		return nullptr;
+	}
 	check(SocketInfoToDuplicate.Socket);
 
 	const FScopedTransaction Transaction(LOCTEXT("CopySocket", "Copy Socket"));
@@ -718,11 +779,19 @@ USkeletalMeshSocket* FEditableSkeleton::DuplicateSocket(const FSelectedSocketInf
 
 int32 FEditableSkeleton::ValidatePreviewAttachedObjects()
 {
+	if (!Skeleton)
+	{
+		return 0;
+	}
 	return Skeleton->ValidatePreviewAttachedObjects();
 }
 
 void FEditableSkeleton::RecreateBoneTree(USkeletalMesh* NewPreviewMesh)
 {
+	if (!Skeleton)
+	{
+		return;
+	}
 	Skeleton->RecreateBoneTree(NewPreviewMesh);
 }
 
@@ -746,7 +815,7 @@ void FEditableSkeleton::HandleAttachAssets(const TArray<UObject*>& InObjects, co
 				SkeletalMesh->GetPreviewAttachedAssetContainer().AddAttachedObject(Object, InAttachToName);
 			}
 		}
-		else
+		else if (Skeleton)
 		{
 			FScopedTransaction Transaction(LOCTEXT("DragDropAttachSkeletonUndo", "Attach Assets to Skeleton"));
 			Skeleton->Modify();
@@ -761,7 +830,7 @@ void FEditableSkeleton::HandleAttachAssets(const TArray<UObject*>& InObjects, co
 
 void FEditableSkeleton::SetBoneTranslationRetargetingModeRecursive(const TArray<FName>& InBoneNames, EBoneTranslationRetargetingMode::Type NewRetargetingMode)
 {
-	if (InBoneNames.Num())
+	if (Skeleton && InBoneNames.Num())
 	{
 		const FScopedTransaction Transaction(LOCTEXT("SetBoneTranslationRetargetingModeRecursive", "Set Bone Translation Retargeting Mode Recursive"));
 		Skeleton->Modify();
@@ -778,7 +847,7 @@ void FEditableSkeleton::SetBoneTranslationRetargetingModeRecursive(const TArray<
 
 void FEditableSkeleton::HandleDeleteAttachedAssets(const TArray<FPreviewAttachedObjectPair>& InAttachedObjects, TSharedPtr<IPersonaPreviewScene> InPreviewScene)
 {
-	if (InAttachedObjects.Num() > 0)
+	if (Skeleton && InAttachedObjects.Num() > 0)
 	{
 		Skeleton->Modify();
 		USkeletalMesh* SkeletalMesh = InPreviewScene->GetPreviewMeshComponent()->GetSkeletalMeshAsset();
@@ -806,6 +875,10 @@ void FEditableSkeleton::HandleDeleteAttachedAssets(const TArray<FPreviewAttached
 
 void FEditableSkeleton::HandleDeleteSockets(const TArray<FSelectedSocketInfo>& InSocketInfo, TSharedPtr<IPersonaPreviewScene> InPreviewScene)
 {
+	if (!Skeleton)
+	{
+		return;
+	}
 	for (const FSelectedSocketInfo& SocketInfo : InSocketInfo)
 	{
 		FName SocketName = SocketInfo.Socket->SocketName;
@@ -853,6 +926,10 @@ void FEditableSkeleton::HandleDeleteSockets(const TArray<FSelectedSocketInfo>& I
 
 void FEditableSkeleton::HandleDeleteVirtualBones(const TArray<FName>& InVirtualBoneInfo, TSharedPtr<class IPersonaPreviewScene> InPreviewScene)
 {
+	if (!Skeleton)
+	{
+		return;
+	}
 	FScopedTransaction Transaction(LOCTEXT("RemoveVirtualBone", "Remove Virtual Bone from Skeleton"));
 	Skeleton->RemoveVirtualBones(InVirtualBoneInfo);
 
@@ -874,6 +951,10 @@ bool FEditableSkeleton::DoesVirtualBoneAlreadyExist(const FString& InVBName) con
 
 void FEditableSkeleton::RenameVirtualBone(const FName OriginalName, const FName InVBName)
 {
+	if (!Skeleton)
+	{
+		return;
+	}
 	FScopedTransaction Transaction(LOCTEXT("RenameVirtualBone", "Rename Virtual Bone in Skeleton"));
 	Skeleton->RenameVirtualBone(OriginalName, InVBName);
 
@@ -882,6 +963,10 @@ void FEditableSkeleton::RenameVirtualBone(const FName OriginalName, const FName
 
 void FEditableSkeleton::RecursiveSetBlendProfileScales(const FName& InBlendProfileName, const TArray<FName>& InBoneNames, float InScaleToSet)
 {
+	if (!Skeleton)
+	{
+		return;
+	}
 	const FScopedTransaction Transaction(LOCTEXT("SetBlendScalesRecursive", "Recursively Set Blend Profile Scales"));
 	Skeleton->Modify();
 
@@ -928,6 +1013,10 @@ TSharedRef<SWidget> FEditableSkeleton::CreateBlendProfilePicker(const FBlendProf
 
 int32 FEditableSkeleton::DeleteAnimNotifies(const TArray<FName>& InNotifyNames, bool bDeleteFromAnimations)
 {
+	if (!Skeleton)
+	{
+		return 0;
+	}
 	int32 NumAnimationsModified = 0;
 	const FScopedTransaction Transaction(LOCTEXT("DeleteAnimNotify", "Delete Anim Notify"));
 	Skeleton->Modify();
@@ -967,6 +1056,10 @@ int32 FEditableSkeleton::DeleteAnimNotifies(const TArray<FName>& InNotifyNames,
 
 int32 FEditableSkeleton::DeleteSyncMarkers(const TArray<FName>& InSyncMarkerNames, bool bDeleteFromAnimations)
 {
+	if (!Skeleton)
+	{
+		return 0;
+	}
 	int32 NumAnimationsModified = 0;
 	const FScopedTransaction Transaction(LOCTEXT("DeleteSyncMarkers", "Delete Sync Markers"));
 	Skeleton->Modify();
@@ -1004,6 +1097,10 @@ int32 FEditableSkeleton::DeleteSyncMarkers(const TArray<FName>& InSyncMarkerName
 
 void FEditableSkeleton::AddNotify(FName NewName)
 {
+	if (!Skeleton)
+	{
+		return;
+	}
 	const FScopedTransaction Transaction(LOCTEXT("AddNewNotifyToSkeleton", "Add New Anim Notify To Skeleton"));
 	Skeleton->Modify();
 	Skeleton->AddNewAnimationNotify(NewName);
@@ -1016,6 +1113,10 @@ void FEditableSkeleton::AddNotify(FName NewName)
 
 void FEditableSkeleton::AddSyncMarker(FName NewName)
 {
+	if (!Skeleton)
+	{
+		return;
+	}
 	const FScopedTransaction Transaction(LOCTEXT("AddNewSyncMarkerToSkeleton", "Add New Sync Marker To Skeleton"));
 	Skeleton->Modify();
 	Skeleton->RegisterMarkerName(NewName);
@@ -1025,6 +1126,10 @@ void FEditableSkeleton::AddSyncMarker(FName NewName)
 
 int32 FEditableSkeleton::RenameNotify(const FName NewName, const FName OldName, bool bRenameInAnimations)
 {
+	if (!Skeleton)
+	{
+		return 0;
+	}
 	int32 NumAnimationsModified = 0;
 	int32 Index = Skeleton->AnimationNotifies.IndexOfByKey(OldName);
 	if(Index != INDEX_NONE)
@@ -1081,6 +1186,10 @@ int32 FEditableSkeleton::RenameNotify(const FName NewName, const FName OldName,
 
 int32 FEditableSkeleton::RenameSyncMarker(const FName NewName, const FName OldName, bool bRenameInAnimations)
 {
+	if (!Skeleton)
+	{
+		return 0;
+	}
 	int32 NumAnimationsModified = 0;
 	int32 Index = Skeleton->ExistingMarkerNames.IndexOfByKey(OldName);
 	if(Index != INDEX_NONE)
@@ -1139,6 +1248,11 @@ void FEditableSkeleton::BroadcastNotifyChanged()
 
 void FEditableSkeleton::GetCompatibleAnimSequences(TArray<struct FAssetData>& OutAssets)
 {
+	if (!Skeleton)
+	{
+		OutAssets.Empty();
+		return;
+	}
 	//Get the skeleton tag to search for
 	FString SkeletonExportName = FAssetData(Skeleton).GetExportTextName();
 
@@ -1162,15 +1276,20 @@ void FEditableSkeleton::GetCompatibleAnimSequences(TArray<struct FAssetData>& Ou
 
 void FEditableSkeleton::RegisterOnSkeletonHierarchyChanged(const USkeleton::FOnSkeletonHierarchyChanged& InDelegate)
 {
+	if (!Skeleton)
+	{
+		return;
+	}
 	Skeleton->RegisterOnSkeletonHierarchyChanged(InDelegate);
 }
 
 void FEditableSkeleton::UnregisterOnSkeletonHierarchyChanged(FDelegateUserObject Thing)
 {
-	if (Skeleton)
+	if (!Skeleton)
 	{
-		Skeleton->UnregisterOnSkeletonHierarchyChanged(Thing);
+		return;
 	}
+	Skeleton->UnregisterOnSkeletonHierarchyChanged(Thing);
 }
 
 void FEditableSkeleton::RegisterOnNotifiesChanged(const FSimpleMulticastDelegate::FDelegate& InDelegate)
@@ -1185,6 +1304,10 @@ void FEditableSkeleton::UnregisterOnNotifiesChanged(FDelegateUserObject Thing)
 
 void FEditableSkeleton::SetPreviewMesh(class USkeletalMesh* InSkeletalMesh)
 {
+	if (!Skeleton)
+	{
+		return;
+	}
 	if (InSkeletalMesh != Skeleton->GetPreviewMesh())
 	{
 		const FScopedTransaction Transaction(LOCTEXT("ChangeSkeletonPreviewMesh", "Change Skeleton Preview Mesh"));
@@ -1194,11 +1317,19 @@ void FEditableSkeleton::SetPreviewMesh(class USkeletalMesh* InSkeletalMesh)
 
 void FEditableSkeleton::LoadAdditionalPreviewSkeletalMeshes()
 {
+	if (!Skeleton)
+	{
+		return;
+	}
 	Skeleton->LoadAdditionalPreviewSkeletalMeshes();
 }
 
 void FEditableSkeleton::SetAdditionalPreviewSkeletalMeshes(class UDataAsset* InPreviewCollectionAsset)
 {
+	if (!Skeleton)
+	{
+		return;
+	}
 	if (!InPreviewCollectionAsset || InPreviewCollectionAsset->GetClass()->ImplementsInterface(UPreviewCollectionInterface::StaticClass()))
 	{
 		const FScopedTransaction Transaction(LOCTEXT("ChangeSkeletonAdditionalMeshes", "Change Skeleton Additional Meshes"));
@@ -1208,6 +1339,10 @@ void FEditableSkeleton::SetAdditionalPreviewSkeletalMeshes(class UDataAsset* InP
 
 void FEditableSkeleton::RenameRetargetSource(const FName InOldName, const FName InNewName)
 {
+	if (!Skeleton)
+	{
+		return;
+	}
 	const FScopedTransaction Transaction(LOCTEXT("RetargetSourceWindow_Rename", "Rename Retarget Source"));
 
 	FReferencePose* Pose = Skeleton->AnimRetargetSources.Find(InOldName);
@@ -1281,6 +1416,10 @@ void FEditableSkeleton::RenameRetargetSource(const FName InOldName, const FName
 
 void FEditableSkeleton::AddRetargetSource(const FName& InName, USkeletalMesh* InReferenceMesh)
 {
+	if (!Skeleton)
+	{
+		return;
+	}
 	// need to verify if the name is unique, if not create unique name
 	int32 IntSuffix = 1;
 	FReferencePose * ExistingPose = NULL;
@@ -1320,6 +1459,10 @@ void FEditableSkeleton::AddRetargetSource(const FName& InName, USkeletalMesh* In
 
 void FEditableSkeleton::DeleteRetargetSources(const TArray<FName>& InRetargetSourceNames)
 {
+	if (!Skeleton)
+	{
+		return;
+	}
 	FAssetRegistryModule& AssetRegistryModule = FModuleManager::LoadModuleChecked<FAssetRegistryModule>(TEXT("AssetRegistry"));
 
 	const FScopedTransaction Transaction(LOCTEXT("RetargetSourceWindow_Delete", "Delete Retarget Source"));
@@ -1403,6 +1546,10 @@ void FEditableSkeleton::DeleteRetargetSources(const TArray<FName>& InRetargetSou
 
 void FEditableSkeleton::RefreshRetargetSources(const TArray<FName>& InRetargetSourceNames)
 {
+	if (!Skeleton)
+	{
+		return;
+	}
 	const FScopedTransaction Transaction(LOCTEXT("RetargetSourceWindow_Refresh", "Refresh Retarget Source"));
 	for (const FName& RetargetSourceName : InRetargetSourceNames)
 	{
@@ -1427,6 +1574,10 @@ void FEditableSkeleton::RefreshRetargetSources(const TArray<FName>& InRetargetSo
 
 void FEditableSkeleton::AddCompatibleSkeleton(const USkeleton* InCompatibleSkeleton)
 {
+	if (!Skeleton)
+	{
+		return;
+	}
 	const FScopedTransaction Transaction(LOCTEXT("AddedCompatibleSkeletons", "Add Compatible Skeleton"));
 	Skeleton->Modify();
 
@@ -1439,6 +1590,10 @@ void FEditableSkeleton::AddCompatibleSkeleton(const USkeleton* InCompatibleSkele
 
 void FEditableSkeleton::RemoveCompatibleSkeleton(const USkeleton* InCompatibleSkeleton)
 {
+	if (!Skeleton)
+	{
+		return;
+	}
 	const FScopedTransaction Transaction(LOCTEXT("RemoveCompatibleSkeletons", "Remove Compatible Skeleton"));
 	Skeleton->Modify();
 
@@ -1451,6 +1606,10 @@ void FEditableSkeleton::RemoveCompatibleSkeleton(const USkeleton* InCompatibleSk
 
 void FEditableSkeleton::RemoveUnusedBones()
 {
+	if (!Skeleton)
+	{
+		return;
+	}
 	TArray<FName> SkeletonBones;
 	FindUnusedBones(SkeletonBones);
 
@@ -1475,6 +1634,11 @@ void FEditableSkeleton::RemoveUnusedBones()
 
 void FEditableSkeleton::FindUnusedBones(TArray<FName>& OutSkeletonBones, bool bSilent /*=false*/)
 {
+	if (!Skeleton)
+	{
+		OutSkeletonBones.Empty();
+		return;
+	}
 	// This operation can take a long time, so unless we're running silently, give the user a heads-up
 	if (!bSilent)
 	{
@@ -1572,11 +1736,19 @@ void FEditableSkeleton::FindUnusedBones(TArray<FName>& OutSkeletonBones, bool bS
 
 void FEditableSkeleton::UpdateSkeletonReferencePose(USkeletalMesh* InSkeletalMesh)
 {
+	if (!Skeleton)
+	{
+		return;
+	}
 	Skeleton->UpdateReferencePoseFromMesh(InSkeletalMesh);
 }
 
 void FEditableSkeleton::RegisterSlotNode(const FName& InSlotName)
 {
+	if (!Skeleton)
+	{
+		return;
+	}
 	if(!Skeleton->ContainsSlotName(InSlotName))
 	{
 		const FScopedTransaction Transaction(LOCTEXT("RegisterSlotNode", "Register Slot Node"));
@@ -1591,6 +1763,10 @@ void FEditableSkeleton::RegisterSlotNode(const FName& InSlotName)
 
 bool FEditableSkeleton::AddSlotGroupName(const FName& InSlotName)
 {
+	if (!Skeleton)
+	{
+		return false;
+	}
 	const FScopedTransaction Transaction(LOCTEXT("AddSlotGroupName", "Add Slot Group Name"));
 	Skeleton->Modify();
 	bool bResult = Skeleton->AddSlotGroupName(InSlotName);
@@ -1600,6 +1776,10 @@ bool FEditableSkeleton::AddSlotGroupName(const FName& InSlotName)
 
 void FEditableSkeleton::SetSlotGroupName(const FName& InSlotName, const FName& InGroupName)
 {
+	if (!Skeleton)
+	{
+		return;
+	}
 	const FScopedTransaction Transaction(LOCTEXT("SetSlotGroupName", "Set Slot Group Name"));
 	Skeleton->Modify();
 	Skeleton->SetSlotGroupName(InSlotName, InGroupName);
@@ -1609,6 +1789,10 @@ void FEditableSkeleton::SetSlotGroupName(const FName& InSlotName, const FName& I
 
 void FEditableSkeleton::DeleteSlotName(const FName& InSlotName)
 {
+	if (!Skeleton)
+	{
+		return;
+	}
 	const FScopedTransaction Transaction(LOCTEXT("DeleteSlotName", "Delete Slot Name"));
 	Skeleton->Modify();
 	Skeleton->RemoveSlotName(InSlotName);
@@ -1618,6 +1802,10 @@ void FEditableSkeleton::DeleteSlotName(const FName& InSlotName)
 
 void FEditableSkeleton::DeleteSlotGroup(const FName& InGroupName)
 {
+	if (!Skeleton)
+	{
+		return;
+	}
 	const FScopedTransaction Transaction(LOCTEXT("DeleteSlotGroup", "Delete Slot Group"));
 	Skeleton->Modify();
 	Skeleton->RemoveSlotGroup(InGroupName);
@@ -1627,6 +1815,10 @@ void FEditableSkeleton::DeleteSlotGroup(const FName& InGroupName)
 
 void FEditableSkeleton::RenameSlotName(const FName InOldSlotName, const FName InNewSlotName)
 {
+	if (!Skeleton)
+	{
+		return;
+	}
 	const FScopedTransaction Transaction(LOCTEXT("RenameSlotName", "Rename Slot Name"));
 	Skeleton->Modify();
 	Skeleton->RenameSlotName(InOldSlotName, InNewSlotName);
diff --git a/Engine/Source/Editor/SkeletonEditor/Public/IEditableSkeleton.h b/Engine/Source/Editor/SkeletonEditor/Public/IEditableSkeleton.h
--- a/Engine/Source/Editor/SkeletonEditor/Public/IEditableSkeleton.h
+++ b/Engine/Source/Editor/SkeletonEditor/Public/IEditableSkeleton.h
@@ -24,10 +24,10 @@ enum class ESocketParentType : int32
 class IEditableSkeleton
 {
 public:
-	/** Get the skeleton. Const - to modify the skeleton you should use the wrapper methods below */
+	/** Get the skeleton. Const - to modify the skeleton you should use the wrapper methods below. Use IsSkeletonValid first, as the Skeleton is not guaranteed to exist. */
 	virtual const class USkeleton& GetSkeleton() const = 0;
 
-	/** Checks to see if the editable skeleton is valid. It is possible to force-delete the USkeleton. */
+	/** Checks to see if the editable skeleton is valid. It is possible for Skeleton to be non-existent (if deleted externally). */
 	virtual bool IsSkeletonValid() const = 0;
 	
 	/** Get the blend profiles that this skeleton currently contains */
diff --git a/Engine/Source/Runtime/Engine/Public/Animation/AnimInstanceProxy.h b/Engine/Source/Runtime/Engine/Public/Animation/AnimInstanceProxy.h
--- a/Engine/Source/Runtime/Engine/Public/Animation/AnimInstanceProxy.h
+++ b/Engine/Source/Runtime/Engine/Public/Animation/AnimInstanceProxy.h
@@ -363,6 +363,12 @@ public:
 		return ActorTransform;
 	}
 
+	/** There are times when the skeleton is not set - for example if it has been deleted and we are in an editor window update. In these situations, use this to check before calling GetSkeleton() */
+	bool IsSkeletonValid() const
+	{
+		return Skeleton != nullptr;
+	}
+
 	/** Get the current skeleton we are using. Note that this will return nullptr outside of pre/post update */
 	USkeleton* GetSkeleton() 
 	{ 
