Skeletal Mesh Editor: The skeleton tree in the Skeletal Mesh Editor no longer unselects children when their parent gets collapsed

#Synced-CL 49930900

diff --git a/Engine/Source/Editor/SkeletalMeshEditor/Private/SkeletalMeshEditor.cpp b/Engine/Source/Editor/SkeletalMeshEditor/Private/SkeletalMeshEditor.cpp
--- a/Engine/Source/Editor/SkeletalMeshEditor/Private/SkeletalMeshEditor.cpp
+++ b/Engine/Source/Editor/SkeletalMeshEditor/Private/SkeletalMeshEditor.cpp
@@ -260,6 +260,7 @@ void FSkeletalMeshEditor::InitSkeletalMeshEditor(const EToolkitMode::Type Mode,
 	SkeletonTreeArgs.OnSelectionChanged = FOnSkeletonTreeSelectionChanged::CreateSP(this, &FSkeletalMeshEditor::HandleSelectionChanged);
 	SkeletonTreeArgs.PreviewScene = PreviewScene;
 	SkeletonTreeArgs.ContextName = GetToolkitFName();
+	SkeletonTreeArgs.bAllowInvisibleItemSelection = true;
 
 	ISkeletonEditorModule& SkeletonEditorModule = FModuleManager::GetModuleChecked<ISkeletonEditorModule>("SkeletonEditor");
 	SkeletonTree = SkeletonEditorModule.CreateSkeletonTree(PersonaToolkit->GetSkeleton(), SkeletonTreeArgs);
diff --git a/Engine/Source/Editor/SkeletonEditor/Private/SSkeletonTree.cpp b/Engine/Source/Editor/SkeletonEditor/Private/SSkeletonTree.cpp
--- a/Engine/Source/Editor/SkeletonEditor/Private/SSkeletonTree.cpp
+++ b/Engine/Source/Editor/SkeletonEditor/Private/SSkeletonTree.cpp
@@ -171,6 +171,7 @@ void SSkeletonTree::Construct(const FArguments& InArgs, const TSharedRef<FEditab
 	bAllowMeshOperations = InSkeletonTreeArgs.bAllowMeshOperations;
 	bAllowSkeletonOperations = InSkeletonTreeArgs.bAllowSkeletonOperations;
 	bShowDebugVisualizationOptions = InSkeletonTreeArgs.bShowDebugVisualizationOptions;
+	bAllowInvisibleItemSelection = InSkeletonTreeArgs.bAllowInvisibleItemSelection;
 	Extenders = InSkeletonTreeArgs.Extenders;
 	OnGetFilterText = InSkeletonTreeArgs.OnGetFilterText;
 	Builder = InSkeletonTreeArgs.Builder;
@@ -679,6 +680,7 @@ void SSkeletonTree::CreateTreeColumns()
 		.OnMouseButtonDoubleClick(this, &SSkeletonTree::OnTreeDoubleClick)
 		.OnSetExpansionRecursive(this, &SSkeletonTree::SetTreeItemExpansionRecursive)
 		.HighlightParentNodesForSelection(true)
+		.AllowInvisibleItemSelection(bAllowInvisibleItemSelection)
 		.HeaderRow
 		(
 			TreeHeaderRow
diff --git a/Engine/Source/Editor/SkeletonEditor/Private/SSkeletonTree.h b/Engine/Source/Editor/SkeletonEditor/Private/SSkeletonTree.h
--- a/Engine/Source/Editor/SkeletonEditor/Private/SSkeletonTree.h
+++ b/Engine/Source/Editor/SkeletonEditor/Private/SSkeletonTree.h
@@ -443,6 +443,9 @@ private:
 	/** Whether to show the filter option to allow filtering of debug draw elements in the viewport. */
 	bool bShowDebugVisualizationOptions;
 
+	/** Wether children can remain selected while their parent is collapsed in the tree */
+	bool bAllowInvisibleItemSelection = false;
+
 	/** Extenders for menus */
 	TSharedPtr<FExtender> Extenders;
 
diff --git a/Engine/Source/Editor/SkeletonEditor/Public/ISkeletonTree.h b/Engine/Source/Editor/SkeletonEditor/Public/ISkeletonTree.h
--- a/Engine/Source/Editor/SkeletonEditor/Public/ISkeletonTree.h
+++ b/Engine/Source/Editor/SkeletonEditor/Public/ISkeletonTree.h
@@ -52,6 +52,7 @@ struct FSkeletonTreeArgs
 		, bAllowMeshOperations(true)
 		, bAllowSkeletonOperations(true)
 		, bHideBonesByDefault(false)
+		, bAllowInvisibleItemSelection(false)
 	{}
 
 	/** Delegate called by the tree when a socket is selected */
@@ -92,6 +93,9 @@ struct FSkeletonTreeArgs
 
 	/** Whether to hide bones by default */
 	bool bHideBonesByDefault;
+
+	/** Wether children can remain selected while their parent is collapsed in the tree */
+	bool bAllowInvisibleItemSelection;
 };
 
 /** Interface used to deal with skeleton editing UI */
