Skeleton tree menu items now respect bAllowSkeletonOperations / bAllowMeshOperations flags

Changes:
- Added FIsActionButtonVisible lambdas to command mappings for socket, blend profile, and virtual bone actions, keying visibility off bAllowSkeletonOperations / bAllowMeshOperations
- Converted the static "Create" and "Blend" sections in CreateNewMenuWidget to dynamic entries that early-out when GetAllowSkeletonOperations() returns false, so the "+" button menu is also correctly filtered
- Added Visibility_Lambda to the toolbar "+" button itself to collapse it entirely when skeleton operations are disallowed
- Extended the existing bAllowSkeletonOperations / bAllowMeshOperations guards in the right-click context  menu to also cover Bone Translation Retargeting, Remove Mesh Socket, and Create Mesh Socket entries
- Exposed GetAllowSkeletonOperations() / GetAllowMeshOperations() accessors on SSkeletonTree to support the dynamic section lambdas

#rnx

#Synced-CL 52270860

diff --git a/Engine/Source/Editor/SkeletonEditor/Private/SSkeletonTree.cpp b/Engine/Source/Editor/SkeletonEditor/Private/SSkeletonTree.cpp
--- a/Engine/Source/Editor/SkeletonEditor/Private/SSkeletonTree.cpp
+++ b/Engine/Source/Editor/SkeletonEditor/Private/SSkeletonTree.cpp
@@ -284,6 +284,7 @@ void SSkeletonTree::Construct(const FArguments& InArgs, const TSharedRef<FEditab
 				.Padding(FMargin(6.f, 0.0))
 				[
 					SNew(SPositiveActionButton)
+					.Visibility_Lambda([this](){ return bAllowSkeletonOperations ? EVisibility::Visible : EVisibility::Collapsed; })
 					.OnGetMenuContent( this, &SSkeletonTree::CreateNewMenuWidget )
 					.Icon(FAppStyle::Get().GetBrush("Icons.Plus"))
 				]
@@ -490,7 +491,9 @@ void SSkeletonTree::BindCommands()
 	CommandList.MapAction(
 		MenuActions.AddSocket,
 		FExecuteAction::CreateSP( this, &SSkeletonTree::OnAddSocket ),
-		FCanExecuteAction::CreateSP( this, &SSkeletonTree::IsAddingSocketsAllowed ) );
+		FCanExecuteAction::CreateSP( this, &SSkeletonTree::IsAddingSocketsAllowed ),
+		FIsActionChecked(),
+		FIsActionButtonVisible::CreateSP(this, &SSkeletonTree::ShouldAllowSkeletonOperations));
 
 	CommandList.MapAction(
 		FGenericCommands::Get().Rename,
@@ -499,16 +502,24 @@ void SSkeletonTree::BindCommands()
 
 	CommandList.MapAction(
 		MenuActions.CreateMeshSocket,
-		FExecuteAction::CreateSP( this, &SSkeletonTree::OnCustomizeSocket ) );
+		FExecuteAction::CreateSP( this, &SSkeletonTree::OnCustomizeSocket ),
+		FCanExecuteAction(),
+		FIsActionChecked(),
+		FIsActionButtonVisible::CreateSP(this, &SSkeletonTree::ShouldAllowMeshOperations));
 
 	CommandList.MapAction(
 		MenuActions.RemoveMeshSocket,
 		FExecuteAction::CreateSP( this, &SSkeletonTree::OnDeleteSelectedRows ),
-		FCanExecuteAction::CreateSP( this, &SSkeletonTree::CanDeleteSelectedRows ) );
+		FCanExecuteAction::CreateSP( this, &SSkeletonTree::CanDeleteSelectedRows ),
+		FIsActionChecked(),
+		FIsActionButtonVisible::CreateSP(this, &SSkeletonTree::ShouldAllowMeshOperations));
 
 	CommandList.MapAction(
 		MenuActions.PromoteSocketToSkeleton,
-		FExecuteAction::CreateSP( this, &SSkeletonTree::OnPromoteSocket ) ); // Adding customization just deletes the mesh socket
+		FExecuteAction::CreateSP( this, &SSkeletonTree::OnPromoteSocket ), // Adding customization just deletes the mesh socket
+		FCanExecuteAction(),
+		FIsActionChecked(),
+		FIsActionButtonVisible::CreateSP(this, &SSkeletonTree::ShouldAllowSkeletonOperations));
 
 	CommandList.MapAction(
 		MenuActions.DeleteSelectedRows,
@@ -549,23 +560,38 @@ void SSkeletonTree::BindCommands()
 
 	CommandList.MapAction(
 		MenuActions.CreateTimeBlendProfile,
-		FExecuteAction::CreateSP( this, &SSkeletonTree::OnCreateBlendProfile, EBlendProfileMode::TimeFactor));
+		FExecuteAction::CreateSP( this, &SSkeletonTree::OnCreateBlendProfile, EBlendProfileMode::TimeFactor),
+		FCanExecuteAction(),
+		FIsActionChecked(),
+		FIsActionButtonVisible::CreateSP(this, &SSkeletonTree::ShouldAllowSkeletonOperations));
 
 	CommandList.MapAction(
 		MenuActions.CreateWeightBlendProfile,
-		FExecuteAction::CreateSP(this, &SSkeletonTree::OnCreateBlendProfile, EBlendProfileMode::WeightFactor));
+		FExecuteAction::CreateSP(this, &SSkeletonTree::OnCreateBlendProfile, EBlendProfileMode::WeightFactor),
+		FCanExecuteAction(),
+		FIsActionChecked(),
+		FIsActionButtonVisible::CreateSP(this, &SSkeletonTree::ShouldAllowSkeletonOperations));
 
 	CommandList.MapAction(
 		MenuActions.CreateBlendMask,
-		FExecuteAction::CreateSP(this, &SSkeletonTree::OnCreateBlendProfile, EBlendProfileMode::BlendMask));
+		FExecuteAction::CreateSP(this, &SSkeletonTree::OnCreateBlendProfile, EBlendProfileMode::BlendMask),
+		FCanExecuteAction(),
+		FIsActionChecked(),
+		FIsActionButtonVisible::CreateSP(this, &SSkeletonTree::ShouldAllowSkeletonOperations));
 
 	CommandList.MapAction(
 		MenuActions.DeleteCurrentBlendProfile,
-		FExecuteAction::CreateSP( this, &SSkeletonTree::OnDeleteCurrentBlendProfile));
+		FExecuteAction::CreateSP( this, &SSkeletonTree::OnDeleteCurrentBlendProfile),
+		FCanExecuteAction(),
+		FIsActionChecked(),
+		FIsActionButtonVisible::CreateSP(this, &SSkeletonTree::ShouldAllowSkeletonOperations));
 
 	CommandList.MapAction(
 		MenuActions.RenameBlendProfile,
-		FExecuteAction::CreateSP(this, &SSkeletonTree::OnRenameBlendProfile));
+		FExecuteAction::CreateSP(this, &SSkeletonTree::OnRenameBlendProfile),
+		FCanExecuteAction(),
+		FIsActionChecked(),
+		FIsActionButtonVisible::CreateSP(this, &SSkeletonTree::ShouldAllowSkeletonOperations));
 
 	PinnedCommands->BindCommandList(UICommandList.ToSharedRef());
 }
@@ -881,7 +907,7 @@ TSharedPtr< SWidget > SSkeletonTree::CreateContextMenu()
 					MenuBuilder.EndSection();
 				}
 
-				if(IsShowingAdvancedOptions())
+				if(IsShowingAdvancedOptions() && bAllowSkeletonOperations)
 				{
 					MenuBuilder.BeginSection("SkeletonTreeBoneTranslationRetargeting", LOCTEXT("BoneTranslationRetargetingHeader", "Bone Translation Retargeting"));
 					{
@@ -959,7 +985,7 @@ TSharedPtr< SWidget > SSkeletonTree::CreateContextMenu()
 			MenuBuilder.EndSection();
 		}
 
-		if(bAllowSkeletonOperations)
+		if (bAllowSkeletonOperations)
 		{
 			if (BoneTreeSelection.HasSelectedOfType<FSkeletonTreeVirtualBoneItem>())
 			{
@@ -985,7 +1011,7 @@ TSharedPtr< SWidget > SSkeletonTree::CreateContextMenu()
 
 					TSharedPtr<FSkeletonTreeSocketItem> SocketItem = BoneTreeSelection.GetSelectedItems<FSkeletonTreeSocketItem>()[0];
 
-					if (SocketItem->IsSocketCustomized() && SocketItem->GetParentType() == ESocketParentType::Mesh )
+					if (SocketItem->IsSocketCustomized() && SocketItem->GetParentType() == ESocketParentType::Mesh && bAllowMeshOperations)
 					{
 						MenuBuilder.AddMenuEntry( Actions.RemoveMeshSocket );
 					}
@@ -994,7 +1020,7 @@ TSharedPtr< SWidget > SSkeletonTree::CreateContextMenu()
 					// and there isn't one of the same name on the mesh, we can customize it
 					if (SocketItem->CanCustomizeSocket() )
 					{
-						if (SocketItem->GetParentType() == ESocketParentType::Skeleton )
+						if (SocketItem->GetParentType() == ESocketParentType::Skeleton && bAllowMeshOperations)
 						{
 							MenuBuilder.AddMenuEntry( Actions.CreateMeshSocket );
 						}
@@ -2008,44 +2034,73 @@ void SSkeletonTree::RegisterNewMenu()
 	}
 
 	FToolMenuOwnerScoped OwnerScoped(this);
-
-	const FSkeletonTreeCommands& Actions = FSkeletonTreeCommands::Get();
-
 	UToolMenu* Menu = UToolMenus::Get()->RegisterMenu(MenuName);
 
 	{
 		FToolMenuSection& CreateSection = Menu->AddSection("CreateNew", LOCTEXT("SkeletonCreateNew", "Create"));
 
-		CreateSection.AddMenuEntry(Actions.AddSocket);
-		CreateSection.AddSubMenu(
-			"VirtualBones",
-			LOCTEXT("AddVirtualBone", "Add Virtual Bone"),
-			LOCTEXT("AddVirtualBone_ToolTip", "Adds a virtual bone to the skeleton."),
-			FNewToolMenuDelegate::CreateLambda([](UToolMenu* InMenu)
-				{
-					USkeletonTreeMenuContext* MenuContext = InMenu->Context.FindContext<USkeletonTreeMenuContext>();
-                    if(MenuContext == nullptr)
-                    {
-                    	return;
-                    }
-                
-                    TSharedPtr<SSkeletonTree> SkeletonTree = MenuContext->SkeletonTree.Pin();
-                    if(!SkeletonTree.IsValid())
-                    {
-                    	return;
-                    }
+		CreateSection.AddDynamicEntry("SocketAndVirtualBone", FNewToolMenuSectionDelegate::CreateLambda([](FToolMenuSection& InSection)
+		{
+			USkeletonTreeMenuContext* Context = InSection.FindContext<USkeletonTreeMenuContext>();
+			if (!Context)
+			{
+				return;
+			}
+
+			TSharedPtr<SSkeletonTree> SkeletonTree = Context->SkeletonTree.Pin();
+			if (!SkeletonTree.IsValid() || !SkeletonTree->ShouldAllowSkeletonOperations())
+			{
+				return;
+			}
+
+			const FSkeletonTreeCommands& Actions = FSkeletonTreeCommands::Get();
+			InSection.AddMenuEntry(Actions.AddSocket);
+			InSection.AddSubMenu(
+				"VirtualBones",
+				LOCTEXT("AddVirtualBone", "Add Virtual Bone"),
+				LOCTEXT("AddVirtualBone_ToolTip", "Adds a virtual bone to the skeleton."),
+				FNewToolMenuDelegate::CreateLambda([](UToolMenu* InMenu)
+					{
+						USkeletonTreeMenuContext* MenuContext = InMenu->Context.FindContext<USkeletonTreeMenuContext>();
+						if (MenuContext == nullptr)
+						{
+							return;
+						}
 
-					TSharedRef<SBoneTreeMenu> MenuContent = SSkeletonTree::CreateVirtualBoneMenu(SkeletonTree.Get());
-					FToolMenuEntry WidgetEntry = FToolMenuEntry::InitWidget("VirtualBones", MenuContent, FText());
-					InMenu->AddMenuEntry(NAME_None, WidgetEntry);
-				}));
+						TSharedPtr<SSkeletonTree> InnerSkeletonTree = MenuContext->SkeletonTree.Pin();
+						if (!InnerSkeletonTree.IsValid())
+						{
+							return;
+						}
+
+						TSharedRef<SBoneTreeMenu> MenuContent = SSkeletonTree::CreateVirtualBoneMenu(InnerSkeletonTree.Get());
+						FToolMenuEntry WidgetEntry = FToolMenuEntry::InitWidget("VirtualBones", MenuContent, FText());
+						InMenu->AddMenuEntry(NAME_None, WidgetEntry);
+					}));
+		}));
 	}
 
 	{
 		FToolMenuSection& BlendSection = Menu->AddSection("Blend", LOCTEXT("SkeletonBlend", "Blend"));
-		BlendSection.AddMenuEntry(Actions.CreateTimeBlendProfile);
-		BlendSection.AddMenuEntry(Actions.CreateWeightBlendProfile);
-		BlendSection.AddMenuEntry(Actions.CreateBlendMask);
+		BlendSection.AddDynamicEntry("BlendProfileEntries", FNewToolMenuSectionDelegate::CreateLambda([](FToolMenuSection& InSection)
+		{
+			USkeletonTreeMenuContext* Context = InSection.FindContext<USkeletonTreeMenuContext>();
+			if (!Context)
+			{
+				return;
+			}
+
+			TSharedPtr<SSkeletonTree> SkeletonTree = Context->SkeletonTree.Pin();
+			if (!SkeletonTree.IsValid() || !SkeletonTree->ShouldAllowSkeletonOperations())
+			{
+				return;
+			}
+
+			const FSkeletonTreeCommands& Actions = FSkeletonTreeCommands::Get();
+			InSection.AddMenuEntry(Actions.CreateTimeBlendProfile);
+			InSection.AddMenuEntry(Actions.CreateWeightBlendProfile);
+			InSection.AddMenuEntry(Actions.CreateBlendMask);
+		}));
 	}
 }
 
diff --git a/Engine/Source/Editor/SkeletonEditor/Private/SSkeletonTree.h b/Engine/Source/Editor/SkeletonEditor/Private/SSkeletonTree.h
--- a/Engine/Source/Editor/SkeletonEditor/Private/SSkeletonTree.h
+++ b/Engine/Source/Editor/SkeletonEditor/Private/SSkeletonTree.h
@@ -100,6 +100,9 @@ public:
 	virtual TSharedPtr<SWidget> GetSearchWidget() const override { return NameFilterBox; }
 	virtual TSharedPtr<IPinnedCommandList> GetPinnedCommandList() const override { return PinnedCommands; }
 
+	bool ShouldAllowSkeletonOperations() const { return bAllowSkeletonOperations; }
+	bool ShouldAllowMeshOperations() const { return bAllowMeshOperations; }
+
 	/** FSelfRegisteringEditorUndoClient interface */
 	virtual void PostUndo(bool bSuccess) override;
 	virtual void PostRedo(bool bSuccess) override;
