Fix off-by-one in Skeleton Tree shift-click range selection

The SSkeletonTreeView override of Private_SelectRangeFromCurrentTo used an
exclusive upper bound (< RangeEndIndex) in its selection loop, causing the
last item in a downward shift-click range to be skipped. Selecting upward
masked the bug because the swapped-out end item was already selected from
the initial click.

Changes:
- Changed loop condition from < to <= to include the end item in the range
- Fixed Clamp max from Num() to Num()-1 to prevent potential out-of-bounds access

#Synced-CL 52538681
#Synced-CL 52538681

diff --git a/Engine/Source/Editor/SkeletonEditor/Private/SSkeletonTree.cpp b/Engine/Source/Editor/SkeletonEditor/Private/SSkeletonTree.cpp
--- a/Engine/Source/Editor/SkeletonEditor/Private/SSkeletonTree.cpp
+++ b/Engine/Source/Editor/SkeletonEditor/Private/SSkeletonTree.cpp
@@ -107,15 +107,15 @@ public:
 
 		int32 RangeEndIndex = ItemsSourceRef.Find( InRangeSelectionEnd );
 
-		RangeStartIndex = FMath::Clamp(RangeStartIndex, 0, ItemsSourceRef.Num());
-		RangeEndIndex = FMath::Clamp(RangeEndIndex, 0, ItemsSourceRef.Num());
+		RangeStartIndex = FMath::Clamp(RangeStartIndex, 0, ItemsSourceRef.Num()-1);
+		RangeEndIndex = FMath::Clamp(RangeEndIndex, 0, ItemsSourceRef.Num()-1);
 
 		if (RangeEndIndex < RangeStartIndex)
 		{
 			Swap( RangeStartIndex, RangeEndIndex );
 		}
 
-		for( int32 ItemIndex = RangeStartIndex; ItemIndex < RangeEndIndex; ++ItemIndex )
+		for( int32 ItemIndex = RangeStartIndex; ItemIndex <= RangeEndIndex; ++ItemIndex )
 		{
 			// check if this item can actually be selected
 			if(Private_CanItemBeSelected(ItemsSourceRef[ItemIndex]))
